AnimNext: Allow OFPA to be optional and able to be changed

External packaging disabled by default for new assets (via factories). External packaging flag defaults to true to be consistent with existing assets.
Add context menu and details panel option to change between external (OFPA) and regular packaging for asset entries.
Works during PIE.
Round tripping works via P4. If an asset was moved from its original location then care must be taken when round-tripping without checking in as the external object name/hash will vary.

Removed 'force save' option from context menu as save is fixed for OFPA assets.

Also fixed an issue where variables outliner entry hash could easily collide (weak object ptr hash is very weak) and some items could be missed.

diff --git a/Engine/Plugins/Experimental/AnimNext/Source/AnimNextEditor/Private/AnimNextEditorModule.cpp b/Engine/Plugins/Experimental/AnimNext/Source/AnimNextEditor/Private/AnimNextEditorModule.cpp
--- a/Engine/Plugins/Experimental/AnimNext/Source/AnimNextEditor/Private/AnimNextEditorModule.cpp
+++ b/Engine/Plugins/Experimental/AnimNext/Source/AnimNextEditor/Private/AnimNextEditorModule.cpp
@@ -46,6 +46,7 @@
 #include "Graph/AnimNextAnimationGraph.h"
 #include "Graph/AnimNextAnimationGraph_EditorData.h"
 #include "Common/AnimNextAssetItemDetails.h"
+#include "Common/AnimNextRigVMAssetEditorDataCustomization.h"
 #include "DataInterface/AnimNextDataInterface.h"
 #include "Entries/AnimNextEventGraphEntry.h"
 #include "Graph/RigVMTrait_AnimNextPublicVariablesEditor.h"
@@ -88,7 +89,10 @@ void FAnimNextEditorModule::StartupModule()
 
 	PropertyModule.RegisterCustomClassLayout("AnimNextVariableEntryProxy", 
 		FOnGetDetailCustomizationInstance::CreateLambda([] { return MakeShared<FVariableProxyCustomization>(); }));
-	
+
+	PropertyModule.RegisterCustomClassLayout("AnimNextRigVMAssetEditorData",
+		FOnGetDetailCustomizationInstance::CreateLambda([] { return MakeShared<FAnimNextRigVMAssetEditorDataCustomization>(); }));
+
 	AnimNextGraphPanelNodeFactory = MakeShared<FAnimNextGraphPanelNodeFactory>();
 	FEdGraphUtilities::RegisterVisualNodeFactory(AnimNextGraphPanelNodeFactory);
 
@@ -200,6 +204,7 @@ void FAnimNextEditorModule::ShutdownModule()
 		PropertyModule.UnregisterCustomPropertyTypeLayout("AnimNextVariableBinding");
 		PropertyModule.UnregisterCustomClassLayout("AnimNextVariableEntry");
 		PropertyModule.UnregisterCustomClassLayout("AnimNextVariableEntryProxy");
+		PropertyModule.UnregisterCustomClassLayout("AnimNextRigVMAssetEditorData");
 	}
 
 	FEdGraphUtilities::UnregisterVisualNodeFactory(AnimNextGraphPanelNodeFactory);
diff --git a/Engine/Plugins/Experimental/AnimNext/Source/AnimNextEditor/Private/Common/AnimNextRigVMAssetEditorDataCustomization.cpp b/Engine/Plugins/Experimental/AnimNext/Source/AnimNextEditor/Private/Common/AnimNextRigVMAssetEditorDataCustomization.cpp
new file mode 100644
--- /dev/null
+++ b/Engine/Plugins/Experimental/AnimNext/Source/AnimNextEditor/Private/Common/AnimNextRigVMAssetEditorDataCustomization.cpp
@@ -0,0 +1,76 @@
+﻿// Copyright Epic Games, Inc. All Rights Reserved.
+
+#include "AnimNextRigVMAssetEditorDataCustomization.h"
+
+#include "AnimNextRigVMAsset.h"
+#include "AnimNextRigVMAssetEditorData.h"
+#include "DetailLayoutBuilder.h"
+#include "DetailWidgetRow.h"
+#include "Widgets/Input/SButton.h"
+
+#define LOCTEXT_NAMESPACE "AnimNextRigVMAssetEditorDataCustomization"
+
+namespace UE::AnimNext::Editor
+{
+
+void FAnimNextRigVMAssetEditorDataCustomization::CustomizeDetails(IDetailLayoutBuilder& DetailBuilder)
+{
+	TSharedRef<IPropertyHandle> ExternalPackagesProperty = DetailBuilder.GetProperty(UAnimNextRigVMAssetEditorData::GetUsesExternalPackagesPropertyName());
+	
+	TArray<TWeakObjectPtr<UObject>> Objects = DetailBuilder.GetSelectedObjects();
+	if(Objects.Num() != 1)
+	{
+		ExternalPackagesProperty->MarkHiddenByCustomization();
+		return;
+	}
+
+	UAnimNextRigVMAsset* Asset = CastChecked<UAnimNextRigVMAsset>(Objects[0]);
+	UAnimNextRigVMAssetEditorData* EditorData = UncookedOnly::FUtils::GetEditorData<UAnimNextRigVMAssetEditorData>(Asset);
+	if(EditorData == nullptr)
+	{
+		ExternalPackagesProperty->MarkHiddenByCustomization();
+		return;
+	}
+
+	IDetailPropertyRow* PropertyRow = DetailBuilder.EditDefaultProperty(ExternalPackagesProperty);
+	if(PropertyRow == nullptr)
+	{
+		ExternalPackagesProperty->MarkHiddenByCustomization();
+		return;
+	}
+
+	PropertyRow->CustomWidget()
+	.NameContent()
+	[
+		ExternalPackagesProperty->CreatePropertyNameWidget()
+	]
+	.ValueContent()
+	[
+		SNew(SButton)
+		.Text_Lambda([EditorData]()
+		{
+			return EditorData->IsUsingExternalPackages() ?
+				LOCTEXT("DisableExternalPackagesLabel", "Use Single Package") :
+				LOCTEXT("EnableExternalPackagesLabel", "Use External Packages");
+		})
+		.ToolTipText_Lambda([EditorData]()
+		{
+			return EditorData->IsUsingExternalPackages() ?
+				LOCTEXT("DisableExternalPackagesTooltip", "Set this asset to use a single package.\nThis will remove any external packages for existing entries, remove them from version control if enabled and save all packages.\nWarning: This operation cannot be undone, so a connection to version control is recommended.") :
+				LOCTEXT("EnableExternalPackagesTooltip", "Set this asset to use external packaging for its entries (graphs, variables etc.)\nThis will create the external packages for all entries, add them to version control if enabled and save all packages.\nWarning: This operation cannot be undone, so a connection to version control is recommended.");
+		})
+		.OnClicked_Lambda([EditorData]()
+		{
+			EditorData->SetUseExternalPackages(!EditorData->IsUsingExternalPackages());
+			return FReply::Handled();
+		})
+	]
+	.ResetToDefaultContent()
+	[
+		SNullWidget::NullWidget
+	];
+}
+
+}
+
+#undef LOCTEXT_NAMESPACE
diff --git a/Engine/Plugins/Experimental/AnimNext/Source/AnimNextEditor/Private/Common/AnimNextRigVMAssetEditorDataCustomization.h b/Engine/Plugins/Experimental/AnimNext/Source/AnimNextEditor/Private/Common/AnimNextRigVMAssetEditorDataCustomization.h
new file mode 100644
--- /dev/null
+++ b/Engine/Plugins/Experimental/AnimNext/Source/AnimNextEditor/Private/Common/AnimNextRigVMAssetEditorDataCustomization.h
@@ -0,0 +1,17 @@
+﻿// Copyright Epic Games, Inc. All Rights Reserved.
+
+#pragma once
+
+#include "IDetailCustomization.h"
+
+namespace UE::AnimNext::Editor
+{
+
+class FAnimNextRigVMAssetEditorDataCustomization : public IDetailCustomization
+{
+private:
+	// IDetailCustomization interface
+	virtual void CustomizeDetails(IDetailLayoutBuilder& DetailBuilder) override;
+};
+
+}
diff --git a/Engine/Plugins/Experimental/AnimNext/Source/AnimNextEditor/Private/DataInterface/AnimNextDataInterfaceAssetDefinition.cpp b/Engine/Plugins/Experimental/AnimNext/Source/AnimNextEditor/Private/DataInterface/AnimNextDataInterfaceAssetDefinition.cpp
--- a/Engine/Plugins/Experimental/AnimNext/Source/AnimNextEditor/Private/DataInterface/AnimNextDataInterfaceAssetDefinition.cpp
+++ b/Engine/Plugins/Experimental/AnimNext/Source/AnimNextEditor/Private/DataInterface/AnimNextDataInterfaceAssetDefinition.cpp
@@ -1,8 +1,14 @@
 ﻿// Copyright Epic Games, Inc. All Rights Reserved.
 
 #include "AnimNextDataInterfaceAssetDefinition.h"
+
+#include "AnimNextRigVMAssetEditorData.h"
 #include "ContentBrowserMenuContexts.h"
+#include "Editor.h"
 #include "IWorkspaceEditorModule.h"
+#include "ToolMenus.h"
+#include "UncookedOnlyUtils.h"
+#include "UObject/SavePackage.h"
 #include "Workspace/AnimNextWorkspaceFactory.h"
 
 #define LOCTEXT_NAMESPACE "AnimNextAssetDefinitions"
@@ -20,4 +26,124 @@ EAssetCommandResult UAssetDefinition_AnimNextDataInterface::OpenAssets(const FAs
 	return EAssetCommandResult::Handled;
 }
 
+namespace UE::AnimNext::Editor
+{
+
+static FDelayedAutoRegisterHelper AutoRegisterGraphMenuItems(EDelayedRegisterRunPhase::EndOfEngineInit, []
+{
+	UToolMenus::RegisterStartupCallback(FSimpleMulticastDelegate::FDelegate::CreateLambda([]()
+	{
+		FToolMenuOwnerScoped OwnerScoped(UE_MODULE_NAME);
+
+		UToolMenu* Menu = UE::ContentBrowser::ExtendToolMenu_AssetContextMenu(UAnimNextDataInterface::StaticClass());
+
+		FToolMenuSection& Section = Menu->FindOrAddSection("GetAssetActions");
+		Section.AddDynamicEntry(NAME_None, FNewToolMenuSectionDelegate::CreateLambda([](FToolMenuSection& InSection)
+		{
+			auto SetExternalPackagingStatus = [](const FToolMenuContext& InContext, bool bInUseExternalPackaging)
+			{
+				UContentBrowserAssetContextMenuContext* AssetContext = InContext.FindContext<UContentBrowserAssetContextMenuContext>();
+				if(AssetContext == nullptr)
+				{
+					return;
+				}
+
+				for (const FAssetData& AssetData : AssetContext->SelectedAssets)
+				{
+					if(!AssetData.GetClass()->IsChildOf(UAnimNextRigVMAsset::StaticClass()))
+					{
+						continue;
+					}
+
+					bool bUsesExternalPackages = true;
+					AssetData.GetTagValue<bool>(UAnimNextRigVMAssetEditorData::GetUsesExternalPackagesPropertyName(), bUsesExternalPackages);
+
+					if(bUsesExternalPackages == bInUseExternalPackaging)
+					{
+						// Status is as requested, skip
+						continue;
+					}
+					
+					UAnimNextRigVMAsset* Asset = Cast<UAnimNextRigVMAsset>(AssetData.GetAsset());
+					if(Asset == nullptr)
+					{
+						continue;
+					}
+
+					UAnimNextRigVMAssetEditorData* EditorData = UncookedOnly::FUtils::GetEditorData<UAnimNextRigVMAssetEditorData>(Asset);
+					if(EditorData == nullptr)
+					{
+						continue;
+					}
+
+					ensure(EditorData->IsUsingExternalPackages() == bUsesExternalPackages);
+					EditorData->SetUseExternalPackages(bInUseExternalPackaging);
+				}
+			};
+
+			auto ExternalPackagingStatusMatches = [](const FToolMenuContext& InContext, bool bInUsesExternalPackaging)
+			{
+				UContentBrowserAssetContextMenuContext* AssetContext = InContext.FindContext<UContentBrowserAssetContextMenuContext>();
+				if(AssetContext == nullptr)
+				{
+					return false;
+				}
+
+				for (const FAssetData& AssetData : AssetContext->SelectedAssets)
+				{
+					if(!AssetData.GetClass()->IsChildOf(UAnimNextRigVMAsset::StaticClass()))
+					{
+						continue;
+					}
+
+					bool bUsesExternalPackages = true;
+					AssetData.GetTagValue<bool>(UAnimNextRigVMAssetEditorData::GetUsesExternalPackagesPropertyName(), bUsesExternalPackages);
+
+					if(bUsesExternalPackages == bInUsesExternalPackaging)
+					{
+						return true;
+					}
+				}
+
+				return false;
+			};
+
+			FToolUIAction UIAction;
+			UIAction.ExecuteAction = FToolMenuExecuteAction::CreateLambda([&SetExternalPackagingStatus](const FToolMenuContext& InContext)
+			{
+				SetExternalPackagingStatus(InContext, true);
+			});
+			UIAction.IsActionVisibleDelegate = FToolMenuIsActionButtonVisible::CreateLambda([&ExternalPackagingStatusMatches](const FToolMenuContext& InContext)
+			{
+				return ExternalPackagingStatusMatches(InContext, false);
+			});
+
+			InSection.AddMenuEntry(
+				"EnableExternalPackages",
+				LOCTEXT("EnableExternalPackagesLabel", "Use External Packages"),
+				LOCTEXT("EnableExternalPackagesTooltip", "Set the asset(s) to use external packaging for its entries (graphs, variables etc.)\nThis will create the external packages for all entries, add them to version control if enabled and save all packages.\nWarning: This operation cannot be undone, so a connection to version control is recommended."),
+				FSlateIcon(FAppStyle::Get().GetStyleSetName(), "Icons.Edit"),
+				UIAction);
+
+			UIAction.ExecuteAction = FToolMenuExecuteAction::CreateLambda([&SetExternalPackagingStatus](const FToolMenuContext& InContext)
+			{
+				SetExternalPackagingStatus(InContext, false);
+			});
+			UIAction.IsActionVisibleDelegate = FToolMenuIsActionButtonVisible::CreateLambda([&ExternalPackagingStatusMatches](const FToolMenuContext& InContext)
+			{
+				return ExternalPackagingStatusMatches(InContext, true);
+			});
+
+			InSection.AddMenuEntry(
+				"DisableExternalPackages",
+				LOCTEXT("DisableExternalPackagesLabel", "Use Single Package"),
+				LOCTEXT("DisableExternalPackagesTooltip", "Set the asset(s) to use a single package.\nThis will remove any external packages for existing entries, remove them from version control if enabled and save all packages.\nWarning: This operation cannot be undone, so a connection to version control is recommended."),
+				FSlateIcon(FAppStyle::Get().GetStyleSetName(), "Icons.Edit"),
+				UIAction);
+		}));
+	}));
+});
+
+}
+
 #undef LOCTEXT_NAMESPACE
\ No newline at end of file
diff --git a/Engine/Plugins/Experimental/AnimNext/Source/AnimNextEditor/Private/DataInterface/AnimNextDataInterfaceFactory.cpp b/Engine/Plugins/Experimental/AnimNext/Source/AnimNextEditor/Private/DataInterface/AnimNextDataInterfaceFactory.cpp
--- a/Engine/Plugins/Experimental/AnimNext/Source/AnimNextEditor/Private/DataInterface/AnimNextDataInterfaceFactory.cpp
+++ b/Engine/Plugins/Experimental/AnimNext/Source/AnimNextEditor/Private/DataInterface/AnimNextDataInterfaceFactory.cpp
@@ -30,6 +30,7 @@ UObject* UAnimNextDataInterfaceFactory::FactoryCreateNew(UClass* Class, UObject*
 	// Create internal editor data
 	UAnimNextDataInterface_EditorData* EditorData = NewObject<UAnimNextDataInterface_EditorData>(NewDataInterface, TEXT("EditorData"), RF_Transactional);
 	NewDataInterface->EditorData = EditorData;
+	EditorData->bUsesExternalPackages = false;
 	EditorData->Initialize(/*bRecompileVM*/false);
 
 	// Compile the initial skeleton
diff --git a/Engine/Plugins/Experimental/AnimNext/Source/AnimNextEditor/Private/Graph/AnimNextAnimationGraphAssetDefinition.cpp b/Engine/Plugins/Experimental/AnimNext/Source/AnimNextEditor/Private/Graph/AnimNextAnimationGraphAssetDefinition.cpp
--- a/Engine/Plugins/Experimental/AnimNext/Source/AnimNextEditor/Private/Graph/AnimNextAnimationGraphAssetDefinition.cpp
+++ b/Engine/Plugins/Experimental/AnimNext/Source/AnimNextEditor/Private/Graph/AnimNextAnimationGraphAssetDefinition.cpp
@@ -1,10 +1,13 @@
 ﻿// Copyright Epic Games, Inc. All Rights Reserved.
 
 #include "AnimNextAnimationGraphAssetDefinition.h"
+
+#include "AnimNextRigVMAssetEditorData.h"
 #include "ContentBrowserMenuContexts.h"
 #include "FileHelpers.h"
 #include "IWorkspaceEditorModule.h"
 #include "ToolMenus.h"
+#include "UncookedOnlyUtils.h"
 #include "Workspace/AnimNextWorkspaceFactory.h"
 
 #define LOCTEXT_NAMESPACE "AnimNextAssetDefinitions"
@@ -23,53 +26,4 @@ EAssetCommandResult UAssetDefinition_AnimNextAnimationGraph::OpenAssets(const FA
 	return EAssetCommandResult::Handled;
 }
 
-namespace UE::AnimNext::Editor
-{
-
-static FDelayedAutoRegisterHelper AutoRegisterGraphMenuItems(EDelayedRegisterRunPhase::EndOfEngineInit, []
-{
-	UToolMenus::RegisterStartupCallback(FSimpleMulticastDelegate::FDelegate::CreateLambda([]()
-	{
-		FToolMenuOwnerScoped OwnerScoped(UE_MODULE_NAME);
-		UToolMenu* Menu = UE::ContentBrowser::ExtendToolMenu_AssetContextMenu(UAnimNextAnimationGraph::StaticClass());
-		 
-		FToolMenuSection& Section = Menu->FindOrAddSection("GetAssetActions");
-		Section.AddDynamicEntry(NAME_None, FNewToolMenuSectionDelegate::CreateLambda([](FToolMenuSection& InSection)
-		{
-			InSection.AddMenuEntry(
-				"ForceSave",
-				LOCTEXT("ForceSaveLabel", "Force Save"),
-				LOCTEXT("ForceSaveTooltip", "Force the save of this item and all its subobjects."),
-				FSlateIcon(FAppStyle::Get().GetStyleSetName(), "Icons.Save"),
-				FToolUIAction(FToolMenuExecuteAction::CreateLambda([](const FToolMenuContext& InContext)
-				{
-					if(UContentBrowserAssetContextMenuContext* AssetContext = InContext.FindContext<UContentBrowserAssetContextMenuContext>())
-					{
-						TArray<UPackage*> PackagesToSave;
-						for (const FAssetData& AssetData : AssetContext->SelectedAssets)
-						{
-							if (UPackage* Package = AssetData.GetPackage())
-							{
-								PackagesToSave.Add(Package);
-								PackagesToSave.Append(Package->GetExternalPackages());
-							}
-						}
-
-						if(PackagesToSave.Num() > 0)
-						{
-							FEditorFileUtils::FPromptForCheckoutAndSaveParams SaveParams;
-							SaveParams.bCheckDirty = false;
-							SaveParams.bPromptToSave = false;
-							SaveParams.bIsExplicitSave = true;
-
-							FEditorFileUtils::PromptForCheckoutAndSave(PackagesToSave, SaveParams);
-						}
-					}
-				})));
-		}));
-	}));
-});
-
-}
-
 #undef LOCTEXT_NAMESPACE
\ No newline at end of file
diff --git a/Engine/Plugins/Experimental/AnimNext/Source/AnimNextEditor/Private/Graph/AnimNextAnimationGraphFactory.cpp b/Engine/Plugins/Experimental/AnimNext/Source/AnimNextEditor/Private/Graph/AnimNextAnimationGraphFactory.cpp
--- a/Engine/Plugins/Experimental/AnimNext/Source/AnimNextEditor/Private/Graph/AnimNextAnimationGraphFactory.cpp
+++ b/Engine/Plugins/Experimental/AnimNext/Source/AnimNextEditor/Private/Graph/AnimNextAnimationGraphFactory.cpp
@@ -31,6 +31,7 @@ UObject* UAnimNextAnimationGraphFactory::FactoryCreateNew(UClass* Class, UObject
 	// Create internal editor data
 	UAnimNextAnimationGraph_EditorData* EditorData = NewObject<UAnimNextAnimationGraph_EditorData>(NewAnimationGraph, TEXT("EditorData"), RF_Transactional);
 	NewAnimationGraph->EditorData = EditorData;
+	EditorData->bUsesExternalPackages = false;
 	EditorData->Initialize(/*bRecompileVM*/false);
 
 	// Add a single internal graph entry
diff --git a/Engine/Plugins/Experimental/AnimNext/Source/AnimNextEditor/Private/Module/AnimNextModuleAssetDefinition.cpp b/Engine/Plugins/Experimental/AnimNext/Source/AnimNextEditor/Private/Module/AnimNextModuleAssetDefinition.cpp
--- a/Engine/Plugins/Experimental/AnimNext/Source/AnimNextEditor/Private/Module/AnimNextModuleAssetDefinition.cpp
+++ b/Engine/Plugins/Experimental/AnimNext/Source/AnimNextEditor/Private/Module/AnimNextModuleAssetDefinition.cpp
@@ -1,11 +1,14 @@
 ﻿// Copyright Epic Games, Inc. All Rights Reserved.
 
 #include "AnimNextModuleAssetDefinition.h"
+
+#include "AnimNextRigVMAssetEditorData.h"
 #include "ContentBrowserMenuContexts.h"
 #include "FileHelpers.h"
 #include "IWorkspaceEditorModule.h"
 #include "IWorkspaceEditor.h"
 #include "ToolMenus.h"
+#include "UncookedOnlyUtils.h"
 #include "Workspace/AnimNextWorkspaceFactory.h"
 
 #define LOCTEXT_NAMESPACE "AssetDefinition_AnimNextGraph"
@@ -24,53 +27,4 @@ EAssetCommandResult UAssetDefinition_AnimNextModule::OpenAssets(const FAssetOpen
 	return EAssetCommandResult::Handled;
 }
 
-namespace UE::AnimNext::Editor
-{
-
-static FDelayedAutoRegisterHelper AutoRegisterModuleMenuItems(EDelayedRegisterRunPhase::EndOfEngineInit, []
-{
-	UToolMenus::RegisterStartupCallback(FSimpleMulticastDelegate::FDelegate::CreateLambda([]()
-	{
-		FToolMenuOwnerScoped OwnerScoped(UE_MODULE_NAME);
-		UToolMenu* Menu = UE::ContentBrowser::ExtendToolMenu_AssetContextMenu(UAnimNextModule::StaticClass());
-		 
-		FToolMenuSection& Section = Menu->FindOrAddSection("GetAssetActions");
-		Section.AddDynamicEntry(NAME_None, FNewToolMenuSectionDelegate::CreateLambda([](FToolMenuSection& InSection)
-		{
-			InSection.AddMenuEntry(
-				"ForceSave",
-				LOCTEXT("ForceSaveLabel", "Force Save"),
-				LOCTEXT("ForceSaveTooltip", "Force the save of this item and all its subobjects."),
-				FSlateIcon(FAppStyle::Get().GetStyleSetName(), "Icons.Save"),
-				FToolUIAction(FToolMenuExecuteAction::CreateLambda([](const FToolMenuContext& InContext)
-				{
-					if(UContentBrowserAssetContextMenuContext* AssetContext = InContext.FindContext<UContentBrowserAssetContextMenuContext>())
-					{
-						TArray<UPackage*> PackagesToSave;
-						for (const FAssetData& AssetData : AssetContext->SelectedAssets)
-						{
-							if (UPackage* Package = AssetData.GetPackage())
-							{
-								PackagesToSave.Add(Package);
-								PackagesToSave.Append(Package->GetExternalPackages());
-							}
-						}
-
-						if(PackagesToSave.Num() > 0)
-						{
-							FEditorFileUtils::FPromptForCheckoutAndSaveParams SaveParams;
-							SaveParams.bCheckDirty = false;
-							SaveParams.bPromptToSave = false;
-							SaveParams.bIsExplicitSave = true;
-
-							FEditorFileUtils::PromptForCheckoutAndSave(PackagesToSave, SaveParams);
-						}
-					}
-				})));
-		}));
-	}));
-});
-
-}
-
 #undef LOCTEXT_NAMESPACE
\ No newline at end of file
diff --git a/Engine/Plugins/Experimental/AnimNext/Source/AnimNextEditor/Private/Module/AnimNextModuleFactory.cpp b/Engine/Plugins/Experimental/AnimNext/Source/AnimNextEditor/Private/Module/AnimNextModuleFactory.cpp
--- a/Engine/Plugins/Experimental/AnimNext/Source/AnimNextEditor/Private/Module/AnimNextModuleFactory.cpp
+++ b/Engine/Plugins/Experimental/AnimNext/Source/AnimNextEditor/Private/Module/AnimNextModuleFactory.cpp
@@ -31,6 +31,7 @@ UObject* UAnimNextModuleFactory::FactoryCreateNew(UClass* Class, UObject* InPare
 	// Create internal editor data
 	UAnimNextModule_EditorData* EditorData = NewObject<UAnimNextModule_EditorData>(NewModule, TEXT("EditorData"), RF_Transactional);
 	NewModule->EditorData = EditorData;
+	EditorData->bUsesExternalPackages = false;
 	EditorData->Initialize(/*bRecompileVM*/false);
 
 	EditorData->AddEventGraph(FRigUnit_AnimNextPrePhysicsEvent::EventName, FRigUnit_AnimNextPrePhysicsEvent::StaticStruct());
diff --git a/Engine/Plugins/Experimental/AnimNext/Source/AnimNextEditor/Private/Variables/Outliner/VariablesOutlinerEntryItem.cpp b/Engine/Plugins/Experimental/AnimNext/Source/AnimNextEditor/Private/Variables/Outliner/VariablesOutlinerEntryItem.cpp
--- a/Engine/Plugins/Experimental/AnimNext/Source/AnimNextEditor/Private/Variables/Outliner/VariablesOutlinerEntryItem.cpp
+++ b/Engine/Plugins/Experimental/AnimNext/Source/AnimNextEditor/Private/Variables/Outliner/VariablesOutlinerEntryItem.cpp
@@ -161,7 +161,7 @@ bool FVariablesOutlinerEntryItem::IsValid() const
 
 FSceneOutlinerTreeItemID FVariablesOutlinerEntryItem::GetID() const
 {
-	return HashCombine(GetTypeHash(WeakEntry), GetTypeHash(WeakDataInterfaceEntry));
+	return HashCombine(GetTypeHash(WeakEntry.Get()), GetTypeHash(WeakDataInterfaceEntry.Get()));
 }
 
 FString FVariablesOutlinerEntryItem::GetDisplayString() const
diff --git a/Engine/Plugins/Experimental/AnimNext/Source/AnimNextUncookedOnly/AnimNextUncookedOnly.Build.cs b/Engine/Plugins/Experimental/AnimNext/Source/AnimNextUncookedOnly/AnimNextUncookedOnly.Build.cs
--- a/Engine/Plugins/Experimental/AnimNext/Source/AnimNextUncookedOnly/AnimNextUncookedOnly.Build.cs
+++ b/Engine/Plugins/Experimental/AnimNext/Source/AnimNextUncookedOnly/AnimNextUncookedOnly.Build.cs
@@ -45,6 +45,7 @@ namespace UnrealBuildTool.Rules
 					new string[]
 					{
 						"UnrealEd",
+						"SourceControl",
 					}
 				);
 
diff --git a/Engine/Plugins/Experimental/AnimNext/Source/AnimNextUncookedOnly/Internal/AnimNextRigVMAssetEditorData.h b/Engine/Plugins/Experimental/AnimNext/Source/AnimNextUncookedOnly/Internal/AnimNextRigVMAssetEditorData.h
--- a/Engine/Plugins/Experimental/AnimNext/Source/AnimNextUncookedOnly/Internal/AnimNextRigVMAssetEditorData.h
+++ b/Engine/Plugins/Experimental/AnimNext/Source/AnimNextUncookedOnly/Internal/AnimNextRigVMAssetEditorData.h
@@ -133,6 +133,17 @@ public:
 	/** Adds a data interface to this asset */
 	UAnimNextDataInterfaceEntry* AddDataInterface(UAnimNextDataInterface* InDataInterface, bool bSetupUndoRedo = true, bool bPrintPythonCommand = true);
 
+	// Get the external packaging status of this asset
+	bool IsUsingExternalPackages() const { return bUsesExternalPackages; }
+
+#if WITH_EDITOR
+	// Switch the external packaging status of this asset
+	void SetUseExternalPackages(bool bInUseExternalPackages);
+
+	// UI helper function
+	static FName GetUsesExternalPackagesPropertyName() { return GET_MEMBER_NAME_CHECKED(UAnimNextRigVMAssetEditorData, bUsesExternalPackages); }
+#endif
+
 	// Report an error to the user, typically used for scripting APIs
 	static void ReportError(const TCHAR* InMessage);
 
@@ -410,10 +421,20 @@ protected:
 	// Handle compiler reporting
 	void HandleReportFromCompiler(EMessageSeverity::Type InSeverity, UObject* InSubject, const FString& InMessage);
 
-	/** All entries in this asset - not saved, either serialized or discovered at load time */
+	// Add a new entry to this asset, taking into account external packaging status
+	void AddEntryInternal(UAnimNextRigVMAssetEntry* InEntry);
+
+	// Remove an entry to this asset, taking into account external packaging status
+	void RemoveEntryInternal(UAnimNextRigVMAssetEntry* InEntry);
+
+	/** All entries in this asset - not saved, discovered at load time and also contains InternalEntries */
 	UPROPERTY(transient)
 	TArray<TObjectPtr<UAnimNextRigVMAssetEntry>> Entries;
 
+	/** Internal entries in this asset */
+	UPROPERTY()
+	TArray<TObjectPtr<UAnimNextRigVMAssetEntry>> InternalEntries;
+
 	UPROPERTY()
 	FRigVMClient RigVMClient;
 
@@ -444,6 +465,9 @@ protected:
 	UPROPERTY(transient, DuplicateTransient)
 	bool bIsCompiling = false;
 
+	UPROPERTY(VisibleAnywhere, Category = "External Packages", AssetRegistrySearchable)
+	bool bUsesExternalPackages = true;
+
 	FOnRigVMCompiledEvent RigVMCompiledEvent;
 
 	FRigVMGraphModifiedEvent RigVMGraphModifiedEvent;
diff --git a/Engine/Plugins/Experimental/AnimNext/Source/AnimNextUncookedOnly/Private/AnimNextRigVMAssetEditorData.cpp b/Engine/Plugins/Experimental/AnimNext/Source/AnimNextUncookedOnly/Private/AnimNextRigVMAssetEditorData.cpp
--- a/Engine/Plugins/Experimental/AnimNext/Source/AnimNextUncookedOnly/Private/AnimNextRigVMAssetEditorData.cpp
+++ b/Engine/Plugins/Experimental/AnimNext/Source/AnimNextUncookedOnly/Private/AnimNextRigVMAssetEditorData.cpp
@@ -14,6 +14,7 @@
 #include "IAnimNextRigVMGraphInterface.h"
 #include "IWorkspaceEditor.h"
 #include "IWorkspaceEditorModule.h"
+#include "ObjectTools.h"
 #include "UncookedOnlyUtils.h"
 #include "DataInterface/AnimNextDataInterface_EditorData.h"
 #include "Entries/AnimNextAnimationGraphEntry.h"
@@ -31,7 +32,12 @@
 #include "RigVMModel/Nodes/RigVMCollapseNode.h"
 #include "UObject/AssetRegistryTagsContext.h"
 #include "DataInterface/AnimNextDataInterface_EditorData.h"
+#include "UObject/SavePackage.h"
 
+#if WITH_EDITOR
+#include "FileHelpers.h"
+#include "PackageSourceControlHelper.h"
+#endif
 
 void UAnimNextRigVMAssetEditorData::BroadcastModified(EAnimNextEditorDataNotifType InType, UObject* InSubject)
 {
@@ -214,6 +220,12 @@ void UAnimNextRigVMAssetEditorData::PostLoadExternalPackages()
 		InLoadedEntry->Initialize(this);
 		Entries.Add(InLoadedEntry);
 	});
+
+	// Internal entries should be empty if we are externally packaged
+	ensure(bUsesExternalPackages == InternalEntries.IsEmpty());
+
+	// Copy any internal entries to the main entries array
+	Entries.Append(InternalEntries);
 }
 
 void UAnimNextRigVMAssetEditorData::PostTransacted(const FTransactionObjectEvent& TransactionEvent)
@@ -711,8 +723,8 @@ TObjectPtr<URigVMGraph> UAnimNextRigVMAssetEditorData::CreateContainedGraphModel
 	URigVMGraph* CollapseNodeModelRootGraph = CollapseNode->GetRootGraph();
 	check(CollapseNodeModelRootGraph);
 
-	// If we are a transient asset, dont use external packages
-	if (!CollapseNodeModelRootGraph->HasAnyFlags(RF_Transient))
+	// If we are a transient asset, or not using external packages dont use external packages
+	if (bUsesExternalPackages && !CollapseNodeModelRootGraph->HasAnyFlags(RF_Transient))
 	{
 		Model->SetExternalPackage(CollapseNodeModelRootGraph->GetExternalPackage());
 	}
@@ -1080,7 +1092,7 @@ bool UAnimNextRigVMAssetEditorData::RemoveEntry(UAnimNextRigVMAssetEntry* InEntr
 	{
 		EntryToRemove->Modify();
 	}
-	Entries.Remove(EntryToRemove);
+	RemoveEntryInternal(EntryToRemove);
 	RefreshExternalModels();
 
 	// This will cause any external package to be removed when saved
@@ -1142,7 +1154,9 @@ UObject* UAnimNextRigVMAssetEditorData::CreateNewSubEntry(UAnimNextRigVMAssetEdi
 	// If we are a transient asset, dont use external packages
 	UAnimNextRigVMAsset* Asset = UE::AnimNext::UncookedOnly::FUtils::GetAsset(InEditorData);
 	check(Asset);
-	if(!Asset->HasAnyFlags(RF_Transient))
+
+	// Additionally check external packaging flag
+	if(!Asset->HasAnyFlags(RF_Transient) && InEditorData->bUsesExternalPackages)
 	{
 		FExternalPackageHelper::SetPackagingMode(NewEntry, InEditorData, true, false, PKG_None);
 	}
@@ -1486,7 +1500,7 @@ UAnimNextVariableEntry* UAnimNextRigVMAssetEditorData::AddVariable(FName InName,
 		Modify();
 	}
 
-	Entries.Add(NewEntry);
+	AddEntryInternal(NewEntry);
 
 	CustomizeNewAssetEntry(NewEntry);
 
@@ -1545,7 +1559,7 @@ UAnimNextEventGraphEntry* UAnimNextRigVMAssetEditorData::AddEventGraph(FName InN
 		Modify();
 	}
 
-	Entries.Add(NewEntry);
+	AddEntryInternal(NewEntry);
 
 	// Add new graph
 	{
@@ -1621,7 +1635,7 @@ UAnimNextAnimationGraphEntry* UAnimNextRigVMAssetEditorData::AddAnimationGraph(F
 		Modify();
 	}
 
-	Entries.Add(NewEntry);
+	AddEntryInternal(NewEntry);
 
 	// Add new graph
 	{
@@ -1767,7 +1781,7 @@ UAnimNextDataInterfaceEntry* UAnimNextRigVMAssetEditorData::AddDataInterface(UAn
 		Modify();
 	}
 
-	Entries.Add(NewEntry);
+	AddEntryInternal(NewEntry);
 
 	CustomizeNewAssetEntry(NewEntry);
 
@@ -1836,3 +1850,127 @@ void UAnimNextRigVMAssetEditorData::RefreshExternalModels()
 		}
 	}
 }
+
+void UAnimNextRigVMAssetEditorData::AddEntryInternal(UAnimNextRigVMAssetEntry* InEntry)
+{
+	// If we are using external packages, dont persist this entry
+	if(bUsesExternalPackages)
+	{
+		Entries.Add(InEntry);
+	}
+	else
+	{
+		InternalEntries.Add(InEntry);
+		Entries.Add(InEntry);
+	}
+}
+
+void UAnimNextRigVMAssetEditorData::RemoveEntryInternal(UAnimNextRigVMAssetEntry* InEntry)
+{
+	if(bUsesExternalPackages)
+	{
+		Entries.Remove(InEntry);
+	}
+	else
+	{
+		InternalEntries.Remove(InEntry);
+		Entries.Remove(InEntry);
+	}
+}
+
+#if WITH_EDITOR
+
+void UAnimNextRigVMAssetEditorData::SetUseExternalPackages(bool bInUseExternalPackages)
+{
+	if(bInUseExternalPackages != bUsesExternalPackages)
+	{
+		if(bInUseExternalPackages)
+		{
+			UPackage* ThisPackage = GetPackage();
+			TArray<UPackage*> AllPackages;
+			AllPackages.Add(ThisPackage);
+
+			// Prompt the user to check out this package, allowing user to decide against this operation
+			if(!FEditorFileUtils::PromptToCheckoutPackages(false, { ThisPackage }))
+			{
+				return;
+			}
+
+			MarkPackageDirty();
+
+			// Set all internal entries to use external packages
+			TArray<UPackage*> ExternalPackages;
+			for (UAnimNextRigVMAssetEntry* Entry : InternalEntries)
+			{
+				FExternalPackageHelper::SetPackagingMode(Entry, this, bInUseExternalPackages, true, PKG_None);
+				UPackage* ExternalPackage = Entry->GetExternalPackage();
+				check(ExternalPackage);
+				ExternalPackages.Add(ExternalPackage);
+				AllPackages.Add(ExternalPackage);
+			}
+
+			// Add (or revert delete) packages to source control
+			FPackageSourceControlHelper SCCHelper;
+			bool bAdded = SCCHelper.AddToSourceControl(ExternalPackages);
+
+			// Clear all internal packages, switch to discovery on PostLoad rather than serialized entries
+			InternalEntries.Empty();
+			bUsesExternalPackages = bInUseExternalPackages;
+
+			// Finally save all packages, they need to be consistent on disk after this operation
+			FEditorFileUtils::FPromptForCheckoutAndSaveParams SaveParams;
+			SaveParams.bAlreadyCheckedOut = true;
+			SaveParams.bCanBeDeclined = false;
+			SaveParams.bPromptToSave = false;
+			FEditorFileUtils::PromptForCheckoutAndSave(AllPackages, SaveParams);
+		}
+		else
+		{
+			// Gather packages we will modify
+			UPackage* ThisPackage = GetPackage();
+			TArray<UPackage*> AllPackages;
+			AllPackages.Add(GetPackage());
+			TArray<UObject*> ExternalPackages;
+			for (UAnimNextRigVMAssetEntry* Entry : Entries)
+			{
+				UPackage* ExternalPackage = Entry->GetExternalPackage();
+				check(ExternalPackage);
+				ExternalPackages.Add(ExternalPackage);
+				AllPackages.Add(ExternalPackage);
+			}
+
+			// Prompt the user to check out files, allowing user to decide against this operation
+			if(!FEditorFileUtils::PromptToCheckoutPackages(false, AllPackages))
+			{
+				return;
+			}
+
+			MarkPackageDirty();
+
+			ensure(InternalEntries.IsEmpty());
+			InternalEntries.Empty();
+
+			// Set all entries to not use external packages
+			for (UAnimNextRigVMAssetEntry* Entry : Entries)
+			{
+				FExternalPackageHelper::SetPackagingMode(Entry, this, bInUseExternalPackages, true, PKG_None);
+			}
+
+			// Delete the old external packages
+			ObjectTools::DeleteObjectsUnchecked(ExternalPackages);
+
+			// Ensure we save all of our entries if we are not using external packages
+			InternalEntries.Append(Entries);
+			bUsesExternalPackages = bInUseExternalPackages;
+
+			// Finally save our package, it needs to be consistent on disk after this operation
+			FEditorFileUtils::FPromptForCheckoutAndSaveParams SaveParams;
+			SaveParams.bAlreadyCheckedOut = true;
+			SaveParams.bCanBeDeclined = false;
+			SaveParams.bPromptToSave = false;
+			FEditorFileUtils::PromptForCheckoutAndSave({ ThisPackage }, SaveParams);
+		}
+	}
+}
+
+#endif
\ No newline at end of file
diff --git a/Engine/Plugins/Experimental/AnimNextStateTree/Source/AnimNextStateTreeEditor/Private/AnimNextStateTreeFactory.cpp b/Engine/Plugins/Experimental/AnimNextStateTree/Source/AnimNextStateTreeEditor/Private/AnimNextStateTreeFactory.cpp
--- a/Engine/Plugins/Experimental/AnimNextStateTree/Source/AnimNextStateTreeEditor/Private/AnimNextStateTreeFactory.cpp
+++ b/Engine/Plugins/Experimental/AnimNextStateTree/Source/AnimNextStateTreeEditor/Private/AnimNextStateTreeFactory.cpp
@@ -39,6 +39,7 @@ UObject* UAnimNextStateTreeFactory::FactoryCreateNew(UClass* Class, UObject* InP
 	// Create internal editor data
 	UAnimNextStateTree_EditorData* EditorData = NewObject<UAnimNextStateTree_EditorData>(NewStateTree, TEXT("EditorData"), RF_Transactional);
 	NewStateTree->EditorData = EditorData;
+	EditorData->bUsesExternalPackages = false;
 	EditorData->Initialize(/*bRecompileVM*/false);
 	
 	// Compile the initial skeleton
