AnimSandbox: Fix for LOD remapping on same mesh

diff --git a/Engine/Plugins/Experimental/AnimNext/Source/AnimNext/Private/Graph/RigUnit_AnimNextRemapPose.cpp b/Engine/Plugins/Experimental/AnimNext/Source/AnimNext/Private/Graph/RigUnit_AnimNextRemapPose.cpp
--- a/Engine/Plugins/Experimental/AnimNext/Source/AnimNext/Private/Graph/RigUnit_AnimNextRemapPose.cpp
+++ b/Engine/Plugins/Experimental/AnimNext/Source/AnimNext/Private/Graph/RigUnit_AnimNextRemapPose.cpp
@@ -36,6 +36,7 @@ FRigUnit_AnimNextRemapPose_Execute()
 
 	if (SourceMesh == TargetMesh)
 	{
+		TargetPose.PrepareForLOD(TargetRefPose, SourcePose.LODLevel, /*bSetRefPose=*/false, SourcePose.IsAdditive());
 		// Just copy the pose to the target in case there is nothing to convert.
 		TargetPose.CopyFrom(SourcePose);
 	}
diff --git a/Engine/Plugins/Experimental/AnimNext/Source/AnimNext/Public/LODPose.h b/Engine/Plugins/Experimental/AnimNext/Source/AnimNext/Public/LODPose.h
--- a/Engine/Plugins/Experimental/AnimNext/Source/AnimNext/Public/LODPose.h
+++ b/Engine/Plugins/Experimental/AnimNext/Source/AnimNext/Public/LODPose.h
@@ -37,7 +37,8 @@ struct FLODPose
 	void CopyFrom(const FLODPose& Other)
 	{
 		check(LODLevel == Other.LODLevel);
-		check(RefPose == Other.RefPose);
+		check(RefPose && Other.RefPose);
+		check(RefPose->ReferenceLocalTransforms.Num() == Other.RefPose->ReferenceLocalTransforms.Num());
 		check(LocalTransformsView.Num() == Other.LocalTransformsView.Num());
 
 		// Copy over the flags and transforms from our source
