AnimNext: Fix drag/drop getting disabled when compile errors are present

Variable actions relied on external variables, which dont get created if compilation fails

diff --git a/Engine/Plugins/Experimental/AnimNext/Source/AnimNextEditor/Private/Common/GraphEditorSchemaActions.cpp b/Engine/Plugins/Experimental/AnimNext/Source/AnimNextEditor/Private/Common/GraphEditorSchemaActions.cpp
--- a/Engine/Plugins/Experimental/AnimNext/Source/AnimNextEditor/Private/Common/GraphEditorSchemaActions.cpp
+++ b/Engine/Plugins/Experimental/AnimNext/Source/AnimNextEditor/Private/Common/GraphEditorSchemaActions.cpp
@@ -139,10 +139,17 @@ UEdGraphNode* FAnimNextSchemaAction_DispatchFactory::PerformAction(UEdGraph* Par
 	return NewNode;
 }
 
-FAnimNextSchemaAction_Variable::FAnimNextSchemaAction_Variable(const FRigVMExternalVariable& InExternalVariable, const EVariableAccessorChoice InVariableAccessorChoice)
-	: ExternalVariable(InExternalVariable)
+FAnimNextSchemaAction_Variable::FAnimNextSchemaAction_Variable(FName InName, const FAnimNextParamType& InType, const EVariableAccessorChoice InVariableAccessorChoice)
+	: Name(InName)
 	, VariableAccessorChoice(InVariableAccessorChoice)
 {
+	if(InType.IsObjectType())
+	{
+		ObjectPath = InType.GetValueTypeObject()->GetPathName();
+	}
+
+	TypeName = InType.ToRigVMTemplateArgument().GetBaseCPPType();
+
 	static const FText VariablesCategory = LOCTEXT("Variables", "Variables");
 	static const FTextFormat GetVariableFormat = LOCTEXT("GetVariableFormat", "Get {0}");
 	static const FTextFormat SetVariableFormat = LOCTEXT("SetVariableFormat", "Set {0}");
@@ -152,28 +159,28 @@ FAnimNextSchemaAction_Variable::FAnimNextSchemaAction_Variable(const FRigVMExter
 
 	if(InVariableAccessorChoice == EVariableAccessorChoice::Get)
 	{
-		MenuDesc = FText::Format(GetVariableFormat, FText::FromName(ExternalVariable.Name));
-		ToolTip = FText::FromString(FString::Printf(TEXT("Get the value of variable %s"), *ExternalVariable.Name.ToString()));
+		MenuDesc = FText::Format(GetVariableFormat, FText::FromName(Name));
+		ToolTip = FText::FromString(FString::Printf(TEXT("Get the value of variable %s"), *Name.ToString()));
 	}
 	else if (InVariableAccessorChoice == EVariableAccessorChoice::Set)
 	{
-		MenuDesc = FText::Format(SetVariableFormat, FText::FromName(ExternalVariable.Name));
-		ToolTip = FText::FromString(FString::Printf(TEXT("Set the value of variable %s"), *ExternalVariable.Name.ToString()));
+		MenuDesc = FText::Format(SetVariableFormat, FText::FromName(Name));
+		ToolTip = FText::FromString(FString::Printf(TEXT("Set the value of variable %s"), *Name.ToString()));
 	}
 	else
 	{
-		MenuDesc = FText::FromName(ExternalVariable.Name);
+		MenuDesc = FText::FromName(Name);
 	}
 
 	UpdateSearchData(MenuDesc, ToolTip, VariablesCategory, FText::GetEmpty());
 
-	FEdGraphPinType PinType = RigVMTypeUtils::PinTypeFromExternalVariable(ExternalVariable);
+	FEdGraphPinType PinType = UE::AnimNext::UncookedOnly::FUtils::GetPinTypeFromParamType(InType);
 	VariableColor = GetDefault<UEdGraphSchema_K2>()->GetPinTypeColor(PinType);
 }
 
 UEdGraphNode* FAnimNextSchemaAction_Variable::PerformAction(UEdGraph* ParentGraph, TArray<UEdGraphPin*>& FromPins, const FVector2D Location, bool bSelectNewNode)
 {
-	const auto AddNodeLambda = [](UEdGraph* ParentGraph, FRigVMExternalVariable ExtVariable, const FVector2D Location, const bool bIsGetter) -> URigVMEdGraphNode*
+	const auto AddNodeLambda = [](UEdGraph* ParentGraph, FName InName, const FString& InTypeName, const FString& InObjectPath, const FVector2D Location, const bool bIsGetter) -> URigVMEdGraphNode*
 	{
 		IRigVMClientHost* Host = ParentGraph->GetImplementingOuter<IRigVMClientHost>();
 		if(Host == nullptr)
@@ -189,24 +196,12 @@ UEdGraphNode* FAnimNextSchemaAction_Variable::PerformAction(UEdGraph* ParentGrap
 
 		URigVMEdGraphNode* NewNode = nullptr;
 
-		FString ObjectPath;
-		if (ExtVariable.TypeObject)
-		{
-			ObjectPath = ExtVariable.TypeObject->GetPathName();
-		}
-
-		FString TypeName = ExtVariable.TypeName.ToString();
-		if (ExtVariable.bIsArray)
-		{
-			TypeName = FString::Printf(TEXT("TArray<%s>"), *TypeName);
-		}
-
 		FString NodeName;
 
-		URigVMController* Controller = Host->GetRigVMClient()->GetController(ParentGraph);	
+		URigVMController* Controller = Host->GetRigVMClient()->GetController(ParentGraph);
 		Controller->OpenUndoBracket(TEXT("Add Variable"));
 
-		if (URigVMNode* ModelNode = Controller->AddVariableNodeFromObjectPath(ExtVariable.Name, TypeName, ObjectPath, bIsGetter, FString(), Location, NodeName, true, true))
+		if (URigVMNode* ModelNode = Controller->AddVariableNodeFromObjectPath(InName, InTypeName, InObjectPath, bIsGetter, FString(), Location, NodeName, true, true))
 		{
 			for (UEdGraphNode* Node : ParentGraph->Nodes)
 			{
@@ -238,31 +233,31 @@ UEdGraphNode* FAnimNextSchemaAction_Variable::PerformAction(UEdGraph* ParentGrap
 	if (VariableAccessorChoice == EVariableAccessorChoice::Deferred)
 	{
 		FMenuBuilder MenuBuilder(true, NULL);
-		const FText SectionText = FText::FromString(FString::Printf(TEXT("Variable %s"), *ExternalVariable.Name.ToString()));
+		const FText SectionText = FText::FromString(FString::Printf(TEXT("Variable %s"), *Name.ToString()));
 
 		MenuBuilder.BeginSection("VariableDropped", SectionText);
 
 		MenuBuilder.AddMenuEntry(
-			FText::FromString(FString::Printf(TEXT("Get %s"), *ExternalVariable.Name.ToString())),
-			FText::FromString(FString::Printf(TEXT("Adds a getter node for variable %s"), *ExternalVariable.Name.ToString())),
+			FText::FromString(FString::Printf(TEXT("Get %s"), *Name.ToString())),
+			FText::FromString(FString::Printf(TEXT("Adds a getter node for variable %s"), *Name.ToString())),
 			FSlateIcon(),
 			FUIAction(
-				FExecuteAction::CreateLambda([AddNodeLambda, ParentGraph, ExtVariable = ExternalVariable, Location]
+				FExecuteAction::CreateLambda([AddNodeLambda, ParentGraph, Location, Name = Name, TypeName = TypeName, ObjectPath = ObjectPath]
 				{
-					AddNodeLambda(ParentGraph, ExtVariable, Location, true);
+					AddNodeLambda(ParentGraph, Name, TypeName, ObjectPath, Location, true);
 				}),
 				FCanExecuteAction()
 			)
 		);
 
 		MenuBuilder.AddMenuEntry(
-			FText::FromString(FString::Printf(TEXT("Set %s"), *ExternalVariable.Name.ToString())),
-			FText::FromString(FString::Printf(TEXT("Adds a setter node for variable %s"), *ExternalVariable.Name.ToString())),
+			FText::FromString(FString::Printf(TEXT("Set %s"), *Name.ToString())),
+			FText::FromString(FString::Printf(TEXT("Adds a setter node for variable %s"), *Name.ToString())),
 			FSlateIcon(),
 			FUIAction(
-				FExecuteAction::CreateLambda([AddNodeLambda, ParentGraph, Location, ExtVariable = ExternalVariable]
+				FExecuteAction::CreateLambda([AddNodeLambda, ParentGraph, Location, Name = Name, TypeName = TypeName, ObjectPath = ObjectPath]
 				{
-					AddNodeLambda(ParentGraph, ExtVariable, Location, false);
+					AddNodeLambda(ParentGraph, Name, TypeName, ObjectPath, Location, false);
 				}),
 				FCanExecuteAction()
 			)
@@ -283,7 +278,7 @@ UEdGraphNode* FAnimNextSchemaAction_Variable::PerformAction(UEdGraph* ParentGrap
 	else
 	{
 		const bool bIsGetter = VariableAccessorChoice == EVariableAccessorChoice::Get;
-		return AddNodeLambda(ParentGraph, ExternalVariable, Location, bIsGetter);
+		return AddNodeLambda(ParentGraph, Name, TypeName, ObjectPath, Location, bIsGetter);
 	}
 }
 
diff --git a/Engine/Plugins/Experimental/AnimNext/Source/AnimNextEditor/Private/Common/GraphEditorSchemaActions.h b/Engine/Plugins/Experimental/AnimNext/Source/AnimNextEditor/Private/Common/GraphEditorSchemaActions.h
--- a/Engine/Plugins/Experimental/AnimNext/Source/AnimNextEditor/Private/Common/GraphEditorSchemaActions.h
+++ b/Engine/Plugins/Experimental/AnimNext/Source/AnimNextEditor/Private/Common/GraphEditorSchemaActions.h
@@ -10,6 +10,7 @@
 #include "Styling/AppStyle.h"
 #include "GraphEditorSchemaActions.generated.h"
 
+struct FAnimNextParamType;
 struct FSlateBrush;
 class URigVMLibraryNode;
 
@@ -99,7 +100,7 @@ struct FAnimNextSchemaAction_Variable : public FAnimNextSchemaAction
 
 	FAnimNextSchemaAction_Variable() = default;
 
-	FAnimNextSchemaAction_Variable(const FRigVMExternalVariable& InExternalVariable, const EVariableAccessorChoice InVariableAccessorChoice);
+	FAnimNextSchemaAction_Variable(FName InName, const FAnimNextParamType& InType, const EVariableAccessorChoice InVariableAccessorChoice);
 
 	// FEdGraphSchemaAction Interface
 	virtual UEdGraphNode* PerformAction(UEdGraph* ParentGraph, TArray<UEdGraphPin*>& FromPins, const FVector2D Location, bool bSelectNewNode = true) override;
@@ -116,7 +117,9 @@ struct FAnimNextSchemaAction_Variable : public FAnimNextSchemaAction
 	}
 
 private:
-	FRigVMExternalVariable ExternalVariable;
+	FName Name;
+	FString ObjectPath;
+	FString TypeName;
 	EVariableAccessorChoice VariableAccessorChoice = EVariableAccessorChoice::Set;
 	FLinearColor VariableColor;
 };
diff --git a/Engine/Plugins/Experimental/AnimNext/Source/AnimNextEditor/Private/Common/SActionMenu.cpp b/Engine/Plugins/Experimental/AnimNext/Source/AnimNextEditor/Private/Common/SActionMenu.cpp
--- a/Engine/Plugins/Experimental/AnimNext/Source/AnimNextEditor/Private/Common/SActionMenu.cpp
+++ b/Engine/Plugins/Experimental/AnimNext/Source/AnimNextEditor/Private/Common/SActionMenu.cpp
@@ -18,6 +18,8 @@
 #include "Units/RigUnit.h"
 #include "UncookedOnlyUtils.h"
 #include "RigVMBlueprintGeneratedClass.h"
+#include "AnimNextRigVMAssetEditorData.h"
+#include "Entries/AnimNextVariableEntry.h"
 
 #define LOCTEXT_NAMESPACE "AnimNextEditor"
 
@@ -148,11 +150,12 @@ void SActionMenu::CollectAllAnimNextGraphActions(FGraphContextMenuBuilder& MenuB
 		}
 	}
 
-	TArray<FRigVMExternalVariable> ExternalVariables = RigVMHost->GetVM()->GetExternalVariables(RigVMHost->GetRigVMExtendedExecuteContext());
-	for (const FRigVMExternalVariable& ExternalVariable : ExternalVariables)
+	TArray<UAnimNextVariableEntry*> Variables;
+	EditorData->GetAllVariables(Variables);
+	for(UAnimNextVariableEntry* VariableEntry : Variables)
 	{
-		MenuBuilder.AddAction(MakeShared<FAnimNextSchemaAction_Variable>(ExternalVariable, FAnimNextSchemaAction_Variable::EVariableAccessorChoice::Set));
-		MenuBuilder.AddAction(MakeShared<FAnimNextSchemaAction_Variable>(ExternalVariable, FAnimNextSchemaAction_Variable::EVariableAccessorChoice::Get));
+		MenuBuilder.AddAction(MakeShared<FAnimNextSchemaAction_Variable>(VariableEntry->GetVariableName(), VariableEntry->GetType(), FAnimNextSchemaAction_Variable::EVariableAccessorChoice::Set));
+		MenuBuilder.AddAction(MakeShared<FAnimNextSchemaAction_Variable>(VariableEntry->GetVariableName(), VariableEntry->GetType(), FAnimNextSchemaAction_Variable::EVariableAccessorChoice::Get));
 	}
 
 	MenuBuilder.AddAction(MakeShared<FAnimNextSchemaAction_AddComment>());
@@ -185,6 +188,8 @@ void SActionMenu::Construct(const FArguments& InArgs, UEdGraph* InGraph)
 	check(RigVMController);
 	RigVMSchema = RigVMController->GetGraph()->GetSchema();
 	check(RigVMSchema);
+	EditorData = Graph->GetTypedOuter<UAnimNextRigVMAssetEditorData>();
+	check(EditorData);
 
 	SBorder::Construct(SBorder::FArguments()
 		.BorderImage(FAppStyle::Get().GetBrush("Menu.Background"))
diff --git a/Engine/Plugins/Experimental/AnimNext/Source/AnimNextEditor/Private/Common/SActionMenu.h b/Engine/Plugins/Experimental/AnimNext/Source/AnimNextEditor/Private/Common/SActionMenu.h
--- a/Engine/Plugins/Experimental/AnimNext/Source/AnimNextEditor/Private/Common/SActionMenu.h
+++ b/Engine/Plugins/Experimental/AnimNext/Source/AnimNextEditor/Private/Common/SActionMenu.h
@@ -12,6 +12,7 @@
 #include "Widgets/Layout/SBorder.h"
 #include "Widgets/SWidget.h"
 
+class UAnimNextRigVMAssetEditorData;
 class URigVMSchema;
 class SEditableTextBox;
 class SGraphActionMenu;
@@ -64,6 +65,7 @@ private:
 	URigVMHost* RigVMHost = nullptr;
 	IRigVMClientHost* RigVMClientHost = nullptr;
 	URigVMController* RigVMController = nullptr;
+	UAnimNextRigVMAssetEditorData* EditorData = nullptr;
 	bool bAutoExpandActionMenu = false;
 	bool bActionExecuted = false;
 
diff --git a/Engine/Plugins/Experimental/AnimNext/Source/AnimNextEditor/Private/Variables/Outliner/VariablesOutlinerMode.cpp b/Engine/Plugins/Experimental/AnimNext/Source/AnimNextEditor/Private/Variables/Outliner/VariablesOutlinerMode.cpp
--- a/Engine/Plugins/Experimental/AnimNext/Source/AnimNextEditor/Private/Variables/Outliner/VariablesOutlinerMode.cpp
+++ b/Engine/Plugins/Experimental/AnimNext/Source/AnimNextEditor/Private/Variables/Outliner/VariablesOutlinerMode.cpp
@@ -355,20 +355,9 @@ TSharedPtr<FDragDropOperation> FVariablesOutlinerMode::CreateDragDropOperation(c
 	{
 		return nullptr;
 	}
-	 
-	FName VariableName = Entry->GetVariableName();
-	TArray<FRigVMExternalVariable> ExternalVariables = RigVMHost->GetVM()->GetExternalVariables(RigVMHost->GetRigVMExtendedExecuteContext());
-	for (const FRigVMExternalVariable& ExternalVariable : ExternalVariables)
-	{
-		if(ExternalVariable.Name == VariableName)
-		{
-			const bool bIsGetter = !(MouseEvent.IsShiftDown() || MouseEvent.IsAltDown());
-			TSharedPtr<FAnimNextSchemaAction_Variable> Action = MakeShared<FAnimNextSchemaAction_Variable>(ExternalVariable, FAnimNextSchemaAction_Variable::EVariableAccessorChoice::Deferred);
-			return FVariableDragDropOp::New(Action);
-		}
-	}
 
-	return nullptr;
+	TSharedPtr<FAnimNextSchemaAction_Variable> Action = MakeShared<FAnimNextSchemaAction_Variable>(Entry->GetVariableName(), Entry->GetType(), FAnimNextSchemaAction_Variable::EVariableAccessorChoice::Deferred);
+	return FVariableDragDropOp::New(Action);
 }
 
 }
diff --git a/Engine/Plugins/Experimental/AnimNext/Source/AnimNextUncookedOnly/Internal/AnimNextRigVMAssetEditorData.h b/Engine/Plugins/Experimental/AnimNext/Source/AnimNextUncookedOnly/Internal/AnimNextRigVMAssetEditorData.h
--- a/Engine/Plugins/Experimental/AnimNext/Source/AnimNextUncookedOnly/Internal/AnimNextRigVMAssetEditorData.h
+++ b/Engine/Plugins/Experimental/AnimNext/Source/AnimNextUncookedOnly/Internal/AnimNextRigVMAssetEditorData.h
@@ -42,6 +42,7 @@ namespace UE::AnimNext::Editor
 	class SAddVariablesDialog;
 	class FVariableProxyCustomization;
 	struct FPublicVariablesImpl;
+	class SActionMenu;
 }
 
 namespace UE::AnimNext::Tests
@@ -175,6 +176,7 @@ protected:
 	friend class UE::AnimNext::Editor::FVariableProxyCustomization;
 	friend class UAnimNextDataInterfaceEntry;
 	friend struct UE::AnimNext::Editor::FPublicVariablesImpl;
+	friend UE::AnimNext::Editor::SActionMenu;
 
 	// UObject interface
 	virtual void Serialize(FArchive& Ar) override;
@@ -415,6 +417,10 @@ protected:
 	// May recurse into implemented data interfaces, so public variables returned by this function may not be directly owned by this asset.
 	void GetPublicVariables(TArray<UAnimNextVariableEntry*>& OutPublicVariables) const;
 
+	// Gets any variables that this asset has. Variables have no specified order.
+	// May recurse into implemented data interfaces (for public variables only), so variables returned by this function may not be directly owned by this asset.
+	void GetAllVariables(TArray<UAnimNextVariableEntry*>& OutPublicVariables) const;
+
 	// Refresh the 'external' models for the RigVM client to reference
 	void RefreshExternalModels();
 
diff --git a/Engine/Plugins/Experimental/AnimNext/Source/AnimNextUncookedOnly/Internal/Entries/AnimNextVariableEntry.h b/Engine/Plugins/Experimental/AnimNext/Source/AnimNextUncookedOnly/Internal/Entries/AnimNextVariableEntry.h
--- a/Engine/Plugins/Experimental/AnimNext/Source/AnimNextUncookedOnly/Internal/Entries/AnimNextVariableEntry.h
+++ b/Engine/Plugins/Experimental/AnimNext/Source/AnimNextUncookedOnly/Internal/Entries/AnimNextVariableEntry.h
@@ -30,6 +30,7 @@ namespace UE::AnimNext::Editor
 	class SVariablesOutlinerValue;
 	class FVariableProxyCustomization;
 	struct FPublicVariablesImpl;
+	class SActionMenu;
 }
 
 namespace UE::AnimNext::Tests
@@ -65,6 +66,7 @@ class UAnimNextVariableEntry : public UAnimNextRigVMAssetEntry, public IAnimNext
 	friend struct UE::AnimNext::UncookedOnly::FUtils;
 	friend class UE::AnimNext::Editor::FVariableProxyCustomization;
 	friend struct UE::AnimNext::Editor::FPublicVariablesImpl;
+	friend class UE::AnimNext::Editor::SActionMenu;
 
 	UAnimNextVariableEntry();
 
diff --git a/Engine/Plugins/Experimental/AnimNext/Source/AnimNextUncookedOnly/Private/AnimNextRigVMAssetEditorData.cpp b/Engine/Plugins/Experimental/AnimNext/Source/AnimNextUncookedOnly/Private/AnimNextRigVMAssetEditorData.cpp
--- a/Engine/Plugins/Experimental/AnimNext/Source/AnimNextUncookedOnly/Private/AnimNextRigVMAssetEditorData.cpp
+++ b/Engine/Plugins/Experimental/AnimNext/Source/AnimNextUncookedOnly/Private/AnimNextRigVMAssetEditorData.cpp
@@ -1835,6 +1835,25 @@ void UAnimNextRigVMAssetEditorData::GetPublicVariables(TArray<UAnimNextVariableE
 	}
 }
 
+void UAnimNextRigVMAssetEditorData::GetAllVariables(TArray<UAnimNextVariableEntry*>& OutPublicVariables) const
+{
+	for(UAnimNextRigVMAssetEntry* Entry : Entries)
+	{
+		if(UAnimNextVariableEntry* VariableEntry = Cast<UAnimNextVariableEntry>(Entry))
+		{
+			OutPublicVariables.Add(VariableEntry);
+		}
+		else if(UAnimNextDataInterfaceEntry* DataInterfaceEntry = Cast<UAnimNextDataInterfaceEntry>(Entry))
+		{
+			if(DataInterfaceEntry->DataInterface)
+			{
+				UAnimNextDataInterface_EditorData* EditorData = UE::AnimNext::UncookedOnly::FUtils::GetEditorData<UAnimNextDataInterface_EditorData>(DataInterfaceEntry->DataInterface.Get());
+				EditorData->GetPublicVariables(OutPublicVariables);
+			}
+		}
+	}
+}
+
 void UAnimNextRigVMAssetEditorData::RefreshExternalModels()
 {
 	GraphModels.Reset();
