Fix Anim Next comment nodes not deleting when selected on their own

diff --git a/Engine/Plugins/Experimental/AnimNext/Source/AnimNextEditor/Private/AnimNextEditorModule.cpp b/Engine/Plugins/Experimental/AnimNext/Source/AnimNextEditor/Private/AnimNextEditorModule.cpp
--- a/Engine/Plugins/Experimental/AnimNext/Source/AnimNextEditor/Private/AnimNextEditorModule.cpp
+++ b/Engine/Plugins/Experimental/AnimNext/Source/AnimNextEditor/Private/AnimNextEditorModule.cpp
@@ -377,16 +377,16 @@ void FAnimNextEditorModule::RegisterWorkspaceDocumentTypes(Workspace::IWorkspace
 				{
 					continue;
 				}
+
+				if(Controller == nullptr)
+				{
+					Controller = RigVMEdGraph->GetController();
+				}
 				
 				if (Node->CanUserDeleteNode())
 				{
 					if (const URigVMEdGraphNode* RigVMEdGraphNode = Cast<URigVMEdGraphNode>(Node))
 					{
-						if(Controller == nullptr)
-						{
-							Controller = RigVMEdGraphNode->GetController();
-						}
-
 						bRelinkPins = bRelinkPins || FSlateApplication::Get().GetModifierKeys().IsShiftDown();
 
 						if(URigVMGraph* Model = RigVMEdGraph->GetModel())
