AnimNext: Fix issues converting OFPA assets

Fixed being unable to save when graphs were created using external packages
Switched to a batch process to better handle converting multiple assets at once.

diff --git a/Engine/Plugins/Experimental/AnimNext/Source/AnimNextEditor/Private/Common/AnimNextRigVMAssetEditorDataCustomization.cpp b/Engine/Plugins/Experimental/AnimNext/Source/AnimNextEditor/Private/Common/AnimNextRigVMAssetEditorDataCustomization.cpp
--- a/Engine/Plugins/Experimental/AnimNext/Source/AnimNextEditor/Private/Common/AnimNextRigVMAssetEditorDataCustomization.cpp
+++ b/Engine/Plugins/Experimental/AnimNext/Source/AnimNextEditor/Private/Common/AnimNextRigVMAssetEditorDataCustomization.cpp
@@ -59,9 +59,9 @@ void FAnimNextRigVMAssetEditorDataCustomization::CustomizeDetails(IDetailLayoutB
 				LOCTEXT("DisableExternalPackagesTooltip", "Set this asset to use a single package.\nThis will remove any external packages for existing entries, remove them from version control if enabled and save all packages.\nWarning: This operation cannot be undone, so a connection to version control is recommended.") :
 				LOCTEXT("EnableExternalPackagesTooltip", "Set this asset to use external packaging for its entries (graphs, variables etc.)\nThis will create the external packages for all entries, add them to version control if enabled and save all packages.\nWarning: This operation cannot be undone, so a connection to version control is recommended.");
 		})
-		.OnClicked_Lambda([EditorData]()
+		.OnClicked_Lambda([Asset, EditorData]() mutable
 		{
-			EditorData->SetUseExternalPackages(!EditorData->IsUsingExternalPackages());
+			UAnimNextRigVMAssetEditorData::SetUseExternalPackages(TArrayView<UAnimNextRigVMAsset*>(&Asset, 1), !EditorData->IsUsingExternalPackages());
 			return FReply::Handled();
 		})
 	]
diff --git a/Engine/Plugins/Experimental/AnimNext/Source/AnimNextEditor/Private/DataInterface/AnimNextDataInterfaceAssetDefinition.cpp b/Engine/Plugins/Experimental/AnimNext/Source/AnimNextEditor/Private/DataInterface/AnimNextDataInterfaceAssetDefinition.cpp
--- a/Engine/Plugins/Experimental/AnimNext/Source/AnimNextEditor/Private/DataInterface/AnimNextDataInterfaceAssetDefinition.cpp
+++ b/Engine/Plugins/Experimental/AnimNext/Source/AnimNextEditor/Private/DataInterface/AnimNextDataInterfaceAssetDefinition.cpp
@@ -48,6 +48,7 @@ static FDelayedAutoRegisterHelper AutoRegisterGraphMenuItems(EDelayedRegisterRun
 					return;
 				}
 
+				TArray<UAnimNextRigVMAsset*> Assets;
 				for (const FAssetData& AssetData : AssetContext->SelectedAssets)
 				{
 					if(!AssetData.GetClass()->IsChildOf(UAnimNextRigVMAsset::StaticClass()))
@@ -70,6 +71,8 @@ static FDelayedAutoRegisterHelper AutoRegisterGraphMenuItems(EDelayedRegisterRun
 						continue;
 					}
 
+					Assets.Add(Asset);
+
 					UAnimNextRigVMAssetEditorData* EditorData = UncookedOnly::FUtils::GetEditorData<UAnimNextRigVMAssetEditorData>(Asset);
 					if(EditorData == nullptr)
 					{
@@ -77,8 +80,9 @@ static FDelayedAutoRegisterHelper AutoRegisterGraphMenuItems(EDelayedRegisterRun
 					}
 
 					ensure(EditorData->IsUsingExternalPackages() == bUsesExternalPackages);
-					EditorData->SetUseExternalPackages(bInUseExternalPackaging);
 				}
+
+				UAnimNextRigVMAssetEditorData::SetUseExternalPackages(Assets, bInUseExternalPackaging);
 			};
 
 			auto ExternalPackagingStatusMatches = [](const FToolMenuContext& InContext, bool bInUsesExternalPackaging)
diff --git a/Engine/Plugins/Experimental/AnimNext/Source/AnimNextUncookedOnly/Internal/AnimNextRigVMAssetEditorData.h b/Engine/Plugins/Experimental/AnimNext/Source/AnimNextUncookedOnly/Internal/AnimNextRigVMAssetEditorData.h
--- a/Engine/Plugins/Experimental/AnimNext/Source/AnimNextUncookedOnly/Internal/AnimNextRigVMAssetEditorData.h
+++ b/Engine/Plugins/Experimental/AnimNext/Source/AnimNextUncookedOnly/Internal/AnimNextRigVMAssetEditorData.h
@@ -139,7 +139,7 @@ public:
 
 #if WITH_EDITOR
 	// Switch the external packaging status of this asset
-	void SetUseExternalPackages(bool bInUseExternalPackages);
+	static void SetUseExternalPackages(TArrayView<UAnimNextRigVMAsset*> InAssets, bool bInUseExternalPackages);
 
 	// UI helper function
 	static FName GetUsesExternalPackagesPropertyName() { return GET_MEMBER_NAME_CHECKED(UAnimNextRigVMAssetEditorData, bUsesExternalPackages); }
diff --git a/Engine/Plugins/Experimental/AnimNext/Source/AnimNextUncookedOnly/Private/AnimNextRigVMAssetEditorData.cpp b/Engine/Plugins/Experimental/AnimNext/Source/AnimNextUncookedOnly/Private/AnimNextRigVMAssetEditorData.cpp
--- a/Engine/Plugins/Experimental/AnimNext/Source/AnimNextUncookedOnly/Private/AnimNextRigVMAssetEditorData.cpp
+++ b/Engine/Plugins/Experimental/AnimNext/Source/AnimNextUncookedOnly/Private/AnimNextRigVMAssetEditorData.cpp
@@ -39,6 +39,8 @@
 #include "PackageSourceControlHelper.h"
 #endif
 
+#define LOCTEXT_NAMESPACE "AnimNextRigVMAssetEditorData"
+
 void UAnimNextRigVMAssetEditorData::BroadcastModified(EAnimNextEditorDataNotifType InType, UObject* InSubject)
 {
 	RecompileVM();
@@ -1899,97 +1901,169 @@ void UAnimNextRigVMAssetEditorData::RemoveEntryInternal(UAnimNextRigVMAssetEntry
 
 #if WITH_EDITOR
 
-void UAnimNextRigVMAssetEditorData::SetUseExternalPackages(bool bInUseExternalPackages)
+void UAnimNextRigVMAssetEditorData::SetUseExternalPackages(TArrayView<UAnimNextRigVMAsset*> InAssets, bool bInUseExternalPackages)
 {
-	if(bInUseExternalPackages != bUsesExternalPackages)
+	TArray<UAnimNextRigVMAssetEditorData*> EditorDatas;
+	for(UAnimNextRigVMAsset* Asset : InAssets)
 	{
-		if(bInUseExternalPackages)
+		if(Asset == nullptr)
 		{
-			UPackage* ThisPackage = GetPackage();
-			TArray<UPackage*> AllPackages;
-			AllPackages.Add(ThisPackage);
+			continue;
+		}
 
-			// Prompt the user to check out this package, allowing user to decide against this operation
-			if(!FEditorFileUtils::PromptToCheckoutPackages(false, { ThisPackage }))
-			{
-				return;
-			}
+		UAnimNextRigVMAssetEditorData* EditorData = UE::AnimNext::UncookedOnly::FUtils::GetEditorData<UAnimNextRigVMAssetEditorData>(Asset);
+		if(EditorData == nullptr)
+		{
+			continue;
+		}
 
-			MarkPackageDirty();
+		if(bInUseExternalPackages != EditorData->bUsesExternalPackages)
+		{
+			EditorDatas.Add(EditorData);
+		}
+	}
+
+	if(EditorDatas.Num() == 0)
+	{
+		return;
+	}
+
+	if(bInUseExternalPackages)
+	{
+		TArray<UPackage*> PackagesToCheckOut;
+		TArray<UPackage*> PackagesToSave;
+		TArray<UPackage*> PackagesToAdd;
+
+		for(UAnimNextRigVMAssetEditorData* EditorData : EditorDatas)
+		{
+			UPackage* Package = EditorData->GetPackage();
+			PackagesToCheckOut.Add(Package);
+			PackagesToSave.Add(Package);
+		}
+
+		// Prompt the user to check out this package, allowing user to decide against this operation
+		if(!FEditorFileUtils::PromptToCheckoutPackages(false, PackagesToCheckOut))
+		{
+			return;
+		}
+
+		FScopedSlowTask SlowTask(3.0f, LOCTEXT("ConvertingAssets", "Converting Assets"));
+		SlowTask.MakeDialog();
+		SlowTask.EnterProgressFrame(1.0f, LOCTEXT("SettingPackagingStatus", "Setting Packaging Status"));
+
+		for(UAnimNextRigVMAssetEditorData* EditorData : EditorDatas)
+		{
+			EditorData->MarkPackageDirty();
 
 			// Set all internal entries to use external packages
 			TArray<UPackage*> ExternalPackages;
-			for (UAnimNextRigVMAssetEntry* Entry : InternalEntries)
+			for (UAnimNextRigVMAssetEntry* Entry : EditorData->InternalEntries)
 			{
-				FExternalPackageHelper::SetPackagingMode(Entry, this, bInUseExternalPackages, true, PKG_None);
+				FExternalPackageHelper::SetPackagingMode(Entry, EditorData, bInUseExternalPackages, true, PKG_None);
 				UPackage* ExternalPackage = Entry->GetExternalPackage();
+
+				// Switch any graphs to be packaged externally
+				if(IAnimNextRigVMGraphInterface* GraphEntry = Cast<IAnimNextRigVMGraphInterface>(Entry))
+				{
+					GraphEntry->GetRigVMGraph()->SetExternalPackage(ExternalPackage);
+				}
+
 				check(ExternalPackage);
-				ExternalPackages.Add(ExternalPackage);
-				AllPackages.Add(ExternalPackage);
+				PackagesToAdd.Add(ExternalPackage);
+				PackagesToSave.Add(ExternalPackage);
 			}
 
-			// Add (or revert delete) packages to source control
-			FPackageSourceControlHelper SCCHelper;
-			bool bAdded = SCCHelper.AddToSourceControl(ExternalPackages);
-
 			// Clear all internal packages, switch to discovery on PostLoad rather than serialized entries
-			InternalEntries.Empty();
-			bUsesExternalPackages = bInUseExternalPackages;
-
-			// Finally save all packages, they need to be consistent on disk after this operation
-			FEditorFileUtils::FPromptForCheckoutAndSaveParams SaveParams;
-			SaveParams.bAlreadyCheckedOut = true;
-			SaveParams.bCanBeDeclined = false;
-			SaveParams.bPromptToSave = false;
-			FEditorFileUtils::PromptForCheckoutAndSave(AllPackages, SaveParams);
+			EditorData->InternalEntries.Empty();
+			EditorData->bUsesExternalPackages = bInUseExternalPackages;
 		}
-		else
+
+		SlowTask.EnterProgressFrame(1.0f, LOCTEXT("AddOrRevertVersionControl", "Adding/Reverting In Version Control"));
+
+		// Add (or revert delete) packages to source control
+		FPackageSourceControlHelper SCCHelper;
+		bool bAdded = SCCHelper.AddToSourceControl(PackagesToAdd);
+
+		SlowTask.EnterProgressFrame(1.0f, LOCTEXT("SavingPackages", "Saving Packages"));
+
+		// Finally save all packages, they need to be consistent on disk after this operation
+		FEditorFileUtils::FPromptForCheckoutAndSaveParams SaveParams;
+		SaveParams.bAlreadyCheckedOut = true;
+		SaveParams.bCanBeDeclined = false;
+		SaveParams.bPromptToSave = false;
+		FEditorFileUtils::PromptForCheckoutAndSave(PackagesToSave, SaveParams);
+	}
+	else
+	{
+		// Gather packages we will modify/delete
+		TArray<UPackage*> PackagesToCheckOut;
+		TArray<UPackage*> PackagesToSave;
+		TArray<UObject*> ObjectsToDelete;
+		for(UAnimNextRigVMAssetEditorData* EditorData : EditorDatas)
 		{
-			// Gather packages we will modify
-			UPackage* ThisPackage = GetPackage();
-			TArray<UPackage*> AllPackages;
-			AllPackages.Add(GetPackage());
-			TArray<UObject*> ExternalPackages;
-			for (UAnimNextRigVMAssetEntry* Entry : Entries)
+			UPackage* ThisPackage = EditorData->GetPackage();
+			PackagesToCheckOut.Add(ThisPackage);
+			for (UAnimNextRigVMAssetEntry* Entry : EditorData->Entries)
 			{
 				UPackage* ExternalPackage = Entry->GetExternalPackage();
 				check(ExternalPackage);
-				ExternalPackages.Add(ExternalPackage);
-				AllPackages.Add(ExternalPackage);
+				ObjectsToDelete.Add(ExternalPackage);
+				PackagesToCheckOut.Add(ExternalPackage);
 			}
+		}
 
-			// Prompt the user to check out files, allowing user to decide against this operation
-			if(!FEditorFileUtils::PromptToCheckoutPackages(false, AllPackages))
-			{
-				return;
-			}
+		// Prompt the user to check out files, allowing user to decide against this operation
+		if(!FEditorFileUtils::PromptToCheckoutPackages(false, PackagesToCheckOut))
+		{
+			return;
+		}
 
-			MarkPackageDirty();
+		FScopedSlowTask SlowTask(3.0f, LOCTEXT("ConvertingAssets", "Converting Assets"));
+		SlowTask.MakeDialog();
+		SlowTask.EnterProgressFrame(1.0f, LOCTEXT("SettingPackagingStatus", "Setting Packaging Status"));
+
+		for(UAnimNextRigVMAssetEditorData* EditorData : EditorDatas)
+		{
+			EditorData->MarkPackageDirty();
 
-			ensure(InternalEntries.IsEmpty());
-			InternalEntries.Empty();
+			ensure(EditorData->InternalEntries.IsEmpty());
+			EditorData->InternalEntries.Empty();
 
 			// Set all entries to not use external packages
-			for (UAnimNextRigVMAssetEntry* Entry : Entries)
+			for (UAnimNextRigVMAssetEntry* Entry : EditorData->Entries)
 			{
-				FExternalPackageHelper::SetPackagingMode(Entry, this, bInUseExternalPackages, true, PKG_None);
-			}
+				FExternalPackageHelper::SetPackagingMode(Entry, EditorData, bInUseExternalPackages, true, PKG_None);
 
-			// Delete the old external packages
-			ObjectTools::DeleteObjectsUnchecked(ExternalPackages);
+				// Switch any graphs to be packaged internally
+				if(IAnimNextRigVMGraphInterface* GraphEntry = Cast<IAnimNextRigVMGraphInterface>(Entry))
+				{
+					GraphEntry->GetRigVMGraph()->SetExternalPackage(nullptr);
+				}
+			}
 
 			// Ensure we save all of our entries if we are not using external packages
-			InternalEntries.Append(Entries);
-			bUsesExternalPackages = bInUseExternalPackages;
+			EditorData->InternalEntries.Append(EditorData->Entries);
+			EditorData->bUsesExternalPackages = bInUseExternalPackages;
 
-			// Finally save our package, it needs to be consistent on disk after this operation
-			FEditorFileUtils::FPromptForCheckoutAndSaveParams SaveParams;
-			SaveParams.bAlreadyCheckedOut = true;
-			SaveParams.bCanBeDeclined = false;
-			SaveParams.bPromptToSave = false;
-			FEditorFileUtils::PromptForCheckoutAndSave({ ThisPackage }, SaveParams);
+			PackagesToSave.Add(EditorData->GetPackage());
 		}
+
+		SlowTask.EnterProgressFrame(1.0f, LOCTEXT("DeletingOldPackages", "Deleting Old Packages"));
+
+		// Delete the old external packages
+		ObjectTools::DeleteObjectsUnchecked(ObjectsToDelete);
+
+		SlowTask.EnterProgressFrame(1.0f, LOCTEXT("SavingPackages", "Saving Packages"));
+
+		// Finally save our packages, they needs to be consistent on disk after this operation
+		FEditorFileUtils::FPromptForCheckoutAndSaveParams SaveParams;
+		SaveParams.bAlreadyCheckedOut = true;
+		SaveParams.bCanBeDeclined = false;
+		SaveParams.bPromptToSave = false;
+		FEditorFileUtils::PromptForCheckoutAndSave(PackagesToSave, SaveParams);
 	}
 }
 
-#endif
\ No newline at end of file
+#endif
+
+#undef LOCTEXT_NAMESPACE
\ No newline at end of file
