AnimNext: Move remaining animation graph code to AnimNextAnimGraph plugin

Required to allow for anim graphs to have knowledge of traits

diff --git a/Engine/Plugins/Experimental/AnimNext/Config/BaseAnimNext.ini b/Engine/Plugins/Experimental/AnimNext/Config/BaseAnimNext.ini
--- a/Engine/Plugins/Experimental/AnimNext/Config/BaseAnimNext.ini
+++ b/Engine/Plugins/Experimental/AnimNext/Config/BaseAnimNext.ini
@@ -65,4 +65,40 @@
 +EnumRedirects=(OldName="EAnimNextParameterFlags",NewName="/Script/AnimNext.EAnimNextExportedVariableFlags")
 
 +PropertyRedirects=(OldName="/Script/AnimNext.RigVMTrait_AnimNextPublicVariables.Asset", NewName="InternalAsset")
-+PropertyRedirects=(OldName="/Script/AnimNext.RigVMTrait_AnimNextPublicVariables.VariableNames", NewName="InternalVariableNames")
\ No newline at end of file
++PropertyRedirects=(OldName="/Script/AnimNext.RigVMTrait_AnimNextPublicVariables.VariableNames", NewName="InternalVariableNames")
+
+; Move all animation graph functionality to AnimNextAnimGraph plugin
++ClassRedirects=(OldName="/Script/AnimNext.AnimNextAnimationGraph", NewName="/Script/AnimNextAnimGraph.AnimNextAnimationGraph")
++StructRedirects=(OldName="/Script/AnimNext.AnimNextGraphContextData", NewName="/Script/AnimNextAnimGraph.AnimNextGraphContextData")
++StructRedirects=(OldName="/Script/AnimNext.AnimNextGraphEntryPoint", NewName="/Script/AnimNextAnimGraph.AnimNextGraphEntryPoint")
++StructRedirects=(OldName="/Script/AnimNext.AnimNextGraphState", NewName="/Script/AnimNextAnimGraph.AnimNextGraphState")
++StructRedirects=(OldName="/Script/AnimNext.RigDecorator_AnimNextCppDecorator", NewName="/Script/AnimNextAnimGraph.RigDecorator_AnimNextCppDecorator")
++StructRedirects=(OldName="/Script/AnimNext.AnimNextCppDecoratorWrapper", NewName="/Script/AnimNextAnimGraph.AnimNextCppDecoratorWrapper")
++StructRedirects=(OldName="/Script/AnimNext.AnimNextGraphEvaluatorExecuteArgument", NewName="/Script/AnimNextAnimGraph.AnimNextGraphEvaluatorExecuteArgument")
++StructRedirects=(OldName="/Script/AnimNext.AnimNextGraphEvaluatorExecuteDefinition", NewName="/Script/AnimNextAnimGraph.AnimNextGraphEvaluatorExecuteDefinition")
++StructRedirects=(OldName="/Script/AnimNext.RigUnit_AnimNextGraphEvaluator", NewName="/Script/AnimNextAnimGraph.RigUnit_AnimNextGraphEvaluator")
++StructRedirects=(OldName="/Script/AnimNext.RigUnit_AnimNextGraphRoot", NewName="/Script/AnimNextAnimGraph.RigUnit_AnimNextGraphRoot")
++StructRedirects=(OldName="/Script/AnimNext.RigUnit_AnimNextRunAnimationGraph", NewName="/Script/AnimNextAnimGraph.RigUnit_AnimNextRunAnimationGraph")
++StructRedirects=(OldName="/Script/AnimNext.RigUnit_AnimNextShimRoot", NewName="/Script/AnimNextAnimGraph.RigUnit_AnimNextShimRoot")
++StructRedirects=(OldName="/Script/AnimNext.RigUnit_AnimNextTraitStack", NewName="/Script/AnimNextAnimGraph.RigUnit_AnimNextTraitStack")
++StructRedirects=(OldName="/Script/AnimNext.AnimNextGraphReference", NewName="/Script/AnimNextAnimGraph.AnimNextGraphReference")
++StructRedirects=(OldName="/Script/AnimNext.AnimNextEntryPointHandle", NewName="/Script/AnimNextAnimGraph.AnimNextEntryPointHandle")
++StructRedirects=(OldName="/Script/AnimNext.AnimNextTraitHandle", NewName="/Script/AnimNextAnimGraph.AnimNextTraitHandle")
++StructRedirects=(OldName="/Script/AnimNext.AnimNextTraitSharedData", NewName="/Script/AnimNextAnimGraph.AnimNextTraitSharedData")
++StructRedirects=(OldName="/Script/AnimNext.AnimNextGraphInstance", NewName="/Script/AnimNextAnimGraph.AnimNextGraphInstance")
++StructRedirects=(OldName="/Script/AnimNext.AnimNextGraphInstancePtr", NewName="/Script/AnimNextAnimGraph.AnimNextGraphInstancePtr")
+
++ClassRedirects=(OldName="/Script/AnimNextEditor.AnimNextAnimationGraphFactory", NewName="/Script/AnimNextAnimGraphEditor.AnimNextAnimationGraphFactory")
++ClassRedirects=(OldName="/Script/AnimNextEditor.AssetDefinition_AnimNextAnimationGraph", NewName="/Script/AnimNextAnimGraphEditor.AssetDefinition_AnimNextAnimationGraph")
+
++ClassRedirects=(OldName="/Script/AnimNextUncookedOnly.AnimNextAnimationGraph_EditorData", NewName="/Script/AnimNextAnimGraphUncookedOnly.AnimNextAnimationGraph_EditorData")
++ClassRedirects=(OldName="/Script/AnimNextUncookedOnly.AnimNextAnimationGraphSchema", NewName="/Script/AnimNextAnimGraphUncookedOnly.AnimNextAnimationGraphSchema")
++StructRedirects=(OldName="/Script/AnimNextUncookedOnly.AnimNextAnimationGraphOutlinerData", NewName="/Script/AnimNextAnimGraphUncookedOnly.AnimNextAnimationGraphOutlinerData")
++ClassRedirects=(OldName="/Script/AnimNextUncookedOnly.AnimNextAnimationGraphEntry", NewName="/Script/AnimNextAnimGraphUncookedOnly.AnimNextAnimationGraphEntry")
+
+
+
+
+
+
+
diff --git a/Engine/Plugins/Experimental/AnimNext/Source/AnimNext/Internal/AnimNextExecuteContext.h b/Engine/Plugins/Experimental/AnimNext/Source/AnimNext/Internal/AnimNextExecuteContext.h
--- a/Engine/Plugins/Experimental/AnimNext/Source/AnimNext/Internal/AnimNextExecuteContext.h
+++ b/Engine/Plugins/Experimental/AnimNext/Source/AnimNext/Internal/AnimNextExecuteContext.h
@@ -5,19 +5,15 @@
 #include "CoreMinimal.h"
 #include "RigVMCore/RigVMExecuteContext.h"
 #include "Module/AnimNextModuleContextData.h"
-#include "Graph/AnimNextGraphContextData.h"
-#include "Graph/AnimNextGraphInstance.h"
-#include "Misc/TVariant.h"
-#include "Module/AnimNextModuleInstance.h"
 #include "AnimNextExecuteContext.generated.h"
 
 struct FAnimNextDataInterfaceInstance;
-struct FAnimNextGraphInstance;
 
 namespace UE::AnimNext
 {
 	struct FLatentPropertyHandle;
 	struct FModuleEventTickFunction;
+	struct FScopedExecuteContextData;
 }
 
 USTRUCT(BlueprintType)
@@ -39,39 +35,42 @@ struct FAnimNextExecuteContext : public FRigVMExecuteContext
 	template<typename ContextType>
 	const ContextType& GetContextData() const
 	{
-		return ContextData.Get<ContextType>();
-	}
-
-	// Get the current data interface instance (module or graph) that is executing 
-	const FAnimNextDataInterfaceInstance& GetInstance() const
-	{
-		if(ContextData.IsType<FAnimNextModuleContextData>())
-		{
-			return ContextData.Get<FAnimNextModuleContextData>().GetModuleInstance();
-		}
-		return ContextData.Get<FAnimNextGraphContextData>().GetGraphInstance();
+		return ContextData.Get<const ContextType>();
 	}
 
 protected:
 	// Setup the context data to the specified type
-	template<typename ContextType, typename... ArgsType>
-	void SetContextData(ArgsType&&... Args)
+	void SetContextData(TStructView<FAnimNextModuleContextData> InContextData)
 	{
-		ContextData.Emplace<ContextType>(Forward<ArgsType>(Args)...);
+		ContextData = InContextData;
 	}
 
-	// Call this to reset the context to its original state to detect stale usage, can't call it Reset due to virtual in base with that name
-	template<typename ContextType>
-	void DebugReset()
+	// Context data for this execution
+	TStructView<FAnimNextModuleContextData> ContextData;
+
+	friend struct UE::AnimNext::FScopedExecuteContextData;
+};
+
+namespace UE::AnimNext
+{
+
+// Helper for applying context data prior to RigVM execution
+struct FScopedExecuteContextData
+{
+	FScopedExecuteContextData() = delete;
+
+	explicit FScopedExecuteContextData(FAnimNextExecuteContext& InContext, TStructView<FAnimNextModuleContextData> InContextData)
+		: Context(InContext)
 	{
-		ContextData.Get<ContextType>().Reset();
+		Context.SetContextData(InContextData);
 	}
 
-	// All possible known variants of our context data. IF we ever want this to be extensible, this can be converted into an FInstancedStruct
-	TVariant<FAnimNextModuleContextData, FAnimNextGraphContextData> ContextData;
+	~FScopedExecuteContextData()
+	{
+		Context.SetContextData(TStructView<FAnimNextModuleContextData>());
+	}
 
-	friend struct UE::AnimNext::FModuleEventTickFunction;
-	friend struct FAnimNextGraphInstance;
-	friend class UAnimNextAnimationGraph;
+	FAnimNextExecuteContext& Context;
 };
 
+}
\ No newline at end of file
diff --git a/Engine/Plugins/Experimental/AnimNext/Source/AnimNext/Internal/AnimNextModuleImpl.h b/Engine/Plugins/Experimental/AnimNext/Source/AnimNext/Internal/AnimNextModuleImpl.h
--- a/Engine/Plugins/Experimental/AnimNext/Source/AnimNext/Internal/AnimNextModuleImpl.h
+++ b/Engine/Plugins/Experimental/AnimNext/Source/AnimNext/Internal/AnimNextModuleImpl.h
@@ -3,41 +3,17 @@
 #pragma once
 
 #include "CoreMinimal.h"
-#include "IAnimNextModuleInterface.h"
-
-struct FAnimNextGraphInstance;
-struct FAnimNextVariableBindingData;
+#include "Modules/ModuleInterface.h"
 
 namespace UE::AnimNext
 {
-	struct FParameterSourceContext;
-}
-
-// Enable console commands only in development builds when logging is enabled
-#define WITH_ANIMNEXT_CONSOLE_COMMANDS (!UE_BUILD_SHIPPING && !NO_LOGGING)
 
-namespace UE::AnimNext
-{
-
-class FAnimNextModuleImpl : public IAnimNextModuleInterface
+class FAnimNextModuleImpl : public IModuleInterface
 {
 public:
 	// IModuleInterface interface
 	virtual void StartupModule() override;
 	virtual void ShutdownModule() override;
-
-	// IAnimNextModuleInterface interface
-	virtual void RegisterAnimNextAnimGraph(const IAnimNextAnimGraph& InAnimGraphImpl) override;
-	virtual void UnregisterAnimNextAnimGraph() override;
-	virtual void UpdateGraph(FAnimNextGraphInstance& GraphInstance, float DeltaTime, FTraitEventList& InputEventList, FTraitEventList& OutputEventList) override;
-	virtual void EvaluateGraph(FAnimNextGraphInstance& GraphInstance, const FReferencePose& RefPose, int32 GraphLODLevel, FAnimNextGraphLODPose& OutputPose) const override;
-
-#if WITH_ANIMNEXT_CONSOLE_COMMANDS
-	void ListNodeTemplates(const TArray<FString>& Args);
-	void ListAnimationGraphs(const TArray<FString>& Args);
-
-	TArray<IConsoleObject*> ConsoleCommands;
-#endif
 };
 
 }
diff --git a/Engine/Plugins/Experimental/AnimNext/Source/AnimNext/Internal/Graph/RigUnit_AnimNextBase.h b/Engine/Plugins/Experimental/AnimNext/Source/AnimNext/Internal/Graph/RigUnit_AnimNextBase.h
--- a/Engine/Plugins/Experimental/AnimNext/Source/AnimNext/Internal/Graph/RigUnit_AnimNextBase.h
+++ b/Engine/Plugins/Experimental/AnimNext/Source/AnimNext/Internal/Graph/RigUnit_AnimNextBase.h
@@ -3,12 +3,12 @@
 #pragma once
 
 #include "CoreMinimal.h"
-#include "Units/RigUnit.h"
+#include "RigVMCore/RigVMStruct.h"
 #include "AnimNextExecuteContext.h"
 #include "RigUnit_AnimNextBase.generated.h"
 
 USTRUCT(meta=(ExecuteContext="FAnimNextExecuteContext"))
-struct FRigUnit_AnimNextBase : public FRigUnit
+struct FRigUnit_AnimNextBase : public FRigVMStruct
 {
 	GENERATED_BODY()
 };
diff --git a/Engine/Plugins/Experimental/AnimNext/Source/AnimNext/Internal/Graph/RigVMTrait_AnimNextPublicVariables.h b/Engine/Plugins/Experimental/AnimNext/Source/AnimNext/Internal/Graph/RigVMTrait_AnimNextPublicVariables.h
--- a/Engine/Plugins/Experimental/AnimNext/Source/AnimNext/Internal/Graph/RigVMTrait_AnimNextPublicVariables.h
+++ b/Engine/Plugins/Experimental/AnimNext/Source/AnimNext/Internal/Graph/RigVMTrait_AnimNextPublicVariables.h
@@ -48,8 +48,8 @@ struct FPublicVariablesTraitToDataInterfaceHostAdapter : public IDataInterfaceHo
 	{}
 
 	// IDataInterfaceHost interface
-	virtual const UAnimNextDataInterface* GetDataInterface() const override;
-	virtual uint8* GetMemoryForVariable(int32 InVariableIndex, FName InVariableName, const FProperty* InVariableProperty) const override;
+	ANIMNEXT_API virtual const UAnimNextDataInterface* GetDataInterface() const override;
+	ANIMNEXT_API virtual uint8* GetMemoryForVariable(int32 InVariableIndex, FName InVariableName, const FProperty* InVariableProperty) const override;
 
 	const FRigVMTrait_AnimNextPublicVariables& Trait;
 	const FRigVMTraitScope& TraitScope;
diff --git a/Engine/Plugins/Experimental/AnimNext/Source/AnimNext/Internal/Module/AnimNextModule.h b/Engine/Plugins/Experimental/AnimNext/Source/AnimNext/Internal/Module/AnimNextModule.h
--- a/Engine/Plugins/Experimental/AnimNext/Source/AnimNext/Internal/Module/AnimNextModule.h
+++ b/Engine/Plugins/Experimental/AnimNext/Source/AnimNext/Internal/Module/AnimNextModule.h
@@ -3,11 +3,10 @@
 #pragma once
 
 #include "CoreMinimal.h"
-#include "Graph/AnimNextGraphEntryPoint.h"
 #include "DataInterface/AnimNextDataInterface.h"
 #include "RigVMCore/RigVM.h"
-#include "TraitCore/TraitPtr.h"
 #include "RigVMHost.h"
+#include "Graph/AnimNextGraphState.h"
 
 #include "AnimNextModule.generated.h"
 
diff --git a/Engine/Plugins/Experimental/AnimNext/Source/AnimNext/Internal/Module/AnimNextModuleContextData.h b/Engine/Plugins/Experimental/AnimNext/Source/AnimNext/Internal/Module/AnimNextModuleContextData.h
--- a/Engine/Plugins/Experimental/AnimNext/Source/AnimNext/Internal/Module/AnimNextModuleContextData.h
+++ b/Engine/Plugins/Experimental/AnimNext/Source/AnimNext/Internal/Module/AnimNextModuleContextData.h
@@ -3,6 +3,8 @@
 #pragma once
 
 #include "CoreMinimal.h"
+#include "AnimNextModuleInstance.h"
+#include "DataInterface/AnimNextDataInterfaceInstance.h"
 #include "AnimNextModuleContextData.generated.h"
 
 struct FAnimNextModuleInstance;
@@ -14,8 +16,15 @@ struct FAnimNextModuleContextData
 
 	FAnimNextModuleContextData() = default;
 
-	FAnimNextModuleContextData(FAnimNextModuleInstance* InModuleInstance)
+	explicit FAnimNextModuleContextData(FAnimNextModuleInstance* InModuleInstance)
 		: ModuleInstance(InModuleInstance)
+		, DataInterfaceInstance(InModuleInstance)
+	{
+	}
+	
+	FAnimNextModuleContextData(FAnimNextModuleInstance* InModuleInstance, const FAnimNextDataInterfaceInstance* InDataInterfaceInstance)
+		: ModuleInstance(InModuleInstance)
+		, DataInterfaceInstance(InDataInterfaceInstance)
 	{
 	}
 
@@ -25,16 +34,19 @@ struct FAnimNextModuleContextData
 		return *ModuleInstance;
 	}
 
-private:
-	// Call this to reset the context to its original state to detect stale usage
-	void Reset()
+	const FAnimNextDataInterfaceInstance& GetDataInterfaceInstance() const
 	{
-		ModuleInstance = nullptr;
+		check(DataInterfaceInstance != nullptr);
+		return *DataInterfaceInstance;
 	}
 
-	// Module instance that is currently executing
+private:
+	// Module instance that is currently executing.
 	FAnimNextModuleInstance* ModuleInstance = nullptr;
 
+	// Data interface that is currently executing. Can be the same as ModuleInstance
+	const FAnimNextDataInterfaceInstance* DataInterfaceInstance = nullptr;
+
 	friend class UAnimNextModule;
 	friend struct FAnimNextExecuteContext;
 };
diff --git a/Engine/Plugins/Experimental/AnimNext/Source/AnimNext/Internal/Module/AnimNextModuleInstance.h b/Engine/Plugins/Experimental/AnimNext/Source/AnimNext/Internal/Module/AnimNextModuleInstance.h
--- a/Engine/Plugins/Experimental/AnimNext/Source/AnimNext/Internal/Module/AnimNextModuleInstance.h
+++ b/Engine/Plugins/Experimental/AnimNext/Source/AnimNext/Internal/Module/AnimNextModuleInstance.h
@@ -111,10 +111,10 @@ struct FAnimNextModuleInstance : public FAnimNextDataInterfaceInstance
 	const ComponentType* TryGetComponent() const;
 
 	// Returns a pointer to the specified component, or nullptr if not found
-	UE::AnimNext::FModuleInstanceComponent* TryGetComponent(int32 ComponentNameHash, FName ComponentName) const;
+	ANIMNEXT_API UE::AnimNext::FModuleInstanceComponent* TryGetComponent(int32 ComponentNameHash, FName ComponentName) const;
 
 	// Adds the specified component and returns a reference to it
-	UE::AnimNext::FModuleInstanceComponent& AddComponent(int32 ComponentNameHash, FName ComponentName, TSharedPtr<UE::AnimNext::FModuleInstanceComponent>&& Component);
+	ANIMNEXT_API UE::AnimNext::FModuleInstanceComponent& AddComponent(int32 ComponentNameHash, FName ComponentName, TSharedPtr<UE::AnimNext::FModuleInstanceComponent>&& Component);
 
 #if UE_ENABLE_DEBUG_DRAWING
 	// Get the debug draw interface
diff --git a/Engine/Plugins/Experimental/AnimNext/Source/AnimNext/Private/AnimNextModuleImpl.cpp b/Engine/Plugins/Experimental/AnimNext/Source/AnimNext/Private/AnimNextModuleImpl.cpp
--- a/Engine/Plugins/Experimental/AnimNext/Source/AnimNext/Private/AnimNextModuleImpl.cpp
+++ b/Engine/Plugins/Experimental/AnimNext/Source/AnimNext/Private/AnimNextModuleImpl.cpp
@@ -10,9 +10,6 @@
 #include "Animation/BlendProfile.h"
 #include "Component/AnimNextComponent.h"
 #include "Curves/CurveFloat.h"
-#include "Graph/AnimNextAnimationGraph.h"
-#include "Graph/AnimNext_LODPose.h"
-#include "Graph/RigVMTrait_AnimNextPublicVariables.h"
 #include "Module/AnimNextModuleInstance.h"
 #include "Modules/ModuleManager.h"
 #include "Param/AnimNextActorLocatorFragment.h"
@@ -22,21 +19,9 @@
 #include "Param/AnimNextObjectPropertyLocatorFragment.h"
 #include "Param/AnimNextTag.h"
 #include "RigVMCore/RigVMRegistry.h"
-#include "TraitCore/NodeTemplateRegistry.h"
-#include "TraitCore/TraitRegistry.h"
-#include "TraitCore/TraitInterfaceRegistry.h"
 #include "Variables/AnimNextFieldPath.h"
 #include "Variables/AnimNextSoftFunctionPtr.h"
 
-#if WITH_ANIMNEXT_CONSOLE_COMMANDS
-#include "HAL/IConsoleManager.h"
-#include "UObject/UObjectIterator.h"
-
-#include "TraitCore/NodeDescription.h"
-#include "TraitCore/NodeTemplate.h"
-#include "TraitCore/TraitTemplate.h"
-#endif
-
 #define LOCTEXT_NAMESPACE "AnimNextModule"
 
 namespace UE::AnimNext
@@ -52,8 +37,8 @@ namespace UE::AnimNext
 			{ UScriptStruct::StaticClass(), FRigVMRegistry::ERegisterObjectOperation::Class },
 			{ UBlendProfile::StaticClass(), FRigVMRegistry::ERegisterObjectOperation::Class },
 			{ UCurveFloat::StaticClass(), FRigVMRegistry::ERegisterObjectOperation::Class },
-			{ UAnimNextAnimationGraph::StaticClass(), FRigVMRegistry::ERegisterObjectOperation::Class },
 			{ UAnimNextComponent::StaticClass(), FRigVMRegistry::ERegisterObjectOperation::Class },
+			{ UAnimNextModule::StaticClass(), FRigVMRegistry::ERegisterObjectOperation::Class },
 		};
 
 		RigVMRegistry.RegisterObjectTypes(AllowedObjectTypes);
@@ -63,19 +48,15 @@ namespace UE::AnimNext
 			FAnimNextScope::StaticStruct(),
 			FAnimNextEntryPoint::StaticStruct(),
 			FUniversalObjectLocator::StaticStruct(),
-			FAnimNextGraphReferencePose::StaticStruct(),
 			FAnimNextFieldPath::StaticStruct(),
 			FAnimNextSoftFunctionPtr::StaticStruct(),
-			FRigVMTrait_AnimNextPublicVariables::StaticStruct()
 		};
 
 		RigVMRegistry.RegisterStructTypes(AllowedStructTypes);
 
 		Private::CacheAllModuleEvents();
 		FDataRegistry::Init();
-		FTraitRegistry::Init();
-		FTraitInterfaceRegistry::Init();
-		FNodeTemplateRegistry::Init();
+
 		FRigVMRuntimeDataRegistry::Init();
 
 		UE::UniversalObjectLocator::IUniversalObjectLocatorModule& UolModule = FModuleManager::Get().LoadModuleChecked<UE::UniversalObjectLocator::IUniversalObjectLocatorModule>("UniversalObjectLocator");
@@ -108,237 +89,14 @@ namespace UE::AnimNext
 					FAnimNextActorLocatorFragment::FragmentType = UolModule.RegisterFragmentType<FAnimNextActorLocatorFragment>(FragmentTypeParams);
 				}
 			});
-#if WITH_ANIMNEXT_CONSOLE_COMMANDS
-		if (!IsRunningCommandlet())
-		{
-			ConsoleCommands.Add(IConsoleManager::Get().RegisterConsoleCommand(
-				TEXT("AnimNext.ListNodeTemplates"),
-				TEXT("Dumps statistics about node templates to the log."),
-				FConsoleCommandWithArgsDelegate::CreateRaw(this, &FAnimNextModuleImpl::ListNodeTemplates),
-				ECVF_Default
-			));
-			ConsoleCommands.Add(IConsoleManager::Get().RegisterConsoleCommand(
-				TEXT("AnimNext.Module"),
-				TEXT("Dumps statistics about modules to the log."),
-				FConsoleCommandWithArgsDelegate::CreateRaw(this, &FAnimNextModuleImpl::ListAnimationGraphs),
-				ECVF_Default
-			));
-		}
-#endif
+
 	}
 
 	void FAnimNextModuleImpl::ShutdownModule()
 	{
 		FRigVMRuntimeDataRegistry::Destroy();
-		FNodeTemplateRegistry::Destroy();
-		FTraitInterfaceRegistry::Destroy();
-		FTraitRegistry::Destroy();
 		FDataRegistry::Destroy();
-
-#if WITH_ANIMNEXT_CONSOLE_COMMANDS
-		for (IConsoleObject* Cmd : ConsoleCommands)
-		{
-			IConsoleManager::Get().UnregisterConsoleObject(Cmd);
-		}
-		ConsoleCommands.Empty();
-#endif
-	}
-
-	const IAnimNextAnimGraph* AnimGraphImpl = nullptr;
-
-	void FAnimNextModuleImpl::RegisterAnimNextAnimGraph(const IAnimNextAnimGraph& InAnimGraphImpl)
-	{
-		AnimGraphImpl = &InAnimGraphImpl;
-	}
-
-	void FAnimNextModuleImpl::UnregisterAnimNextAnimGraph()
-	{
-		AnimGraphImpl = nullptr;
-	}
-
-	void FAnimNextModuleImpl::UpdateGraph(FAnimNextGraphInstance& GraphInstance, float DeltaTime, FTraitEventList& InputEventList, FTraitEventList& OutputEventList)
-	{
-		if (AnimGraphImpl != nullptr)
-		{
-			AnimGraphImpl->UpdateGraph(GraphInstance, DeltaTime, InputEventList, OutputEventList);
-		}
-	}
-
-	void FAnimNextModuleImpl::EvaluateGraph(FAnimNextGraphInstance& GraphInstance, const UE::AnimNext::FReferencePose& RefPose, int32 GraphLODLevel, FAnimNextGraphLODPose& OutputPose) const
-	{
-		if (AnimGraphImpl != nullptr)
-		{
-			AnimGraphImpl->EvaluateGraph(GraphInstance, RefPose, GraphLODLevel, OutputPose);
-		}
-	}
-
-#if WITH_ANIMNEXT_CONSOLE_COMMANDS
-	void FAnimNextModuleImpl::ListNodeTemplates(const TArray<FString>& Args)
-	{
-		// Turn off log times to make diff-ing easier
-		TGuardValue<ELogTimes::Type> DisableLogTimes(GPrintLogTimes, ELogTimes::None);
-
-		// Make sure to log everything
-		const ELogVerbosity::Type OldVerbosity = LogAnimation.GetVerbosity();
-		LogAnimation.SetVerbosity(ELogVerbosity::All);
-
-		const FNodeTemplateRegistry& NodeTemplateRegistry = FNodeTemplateRegistry::Get();
-		const FTraitRegistry& TraitRegistry = FTraitRegistry::Get();
-
-		UE_LOG(LogAnimation, Log, TEXT("===== AnimNext Node Templates ====="));
-		UE_LOG(LogAnimation, Log, TEXT("Template Buffer Size: %u bytes"), NodeTemplateRegistry.TemplateBuffer.GetAllocatedSize());
-
-		for (auto It = NodeTemplateRegistry.TemplateUIDToHandleMap.CreateConstIterator(); It; ++It)
-		{
-			const FNodeTemplateRegistryHandle Handle = It.Value();
-			const FNodeTemplate* NodeTemplate = NodeTemplateRegistry.Find(Handle);
-
-			const uint32 NumTraits = NodeTemplate->GetNumTraits();
-
-			UE_LOG(LogAnimation, Log, TEXT("[%x] has %u traits ..."), NodeTemplate->GetUID(), NumTraits);
-			UE_LOG(LogAnimation, Log, TEXT("    Template Size: %u bytes"), NodeTemplate->GetNodeTemplateSize());
-			UE_LOG(LogAnimation, Log, TEXT("    Shared Data Size: %u bytes"), NodeTemplate->GetNodeSharedDataSize());
-			UE_LOG(LogAnimation, Log, TEXT("    Instance Data Size: %u bytes"), NodeTemplate->GetNodeInstanceDataSize());
-			UE_LOG(LogAnimation, Log, TEXT("    Traits ..."));
-
-			const FTraitTemplate* TraitTemplates = NodeTemplate->GetTraits();
-			for (uint32 TraitIndex = 0; TraitIndex < NumTraits; ++TraitIndex)
-			{
-				const FTraitTemplate* TraitTemplate = TraitTemplates + TraitIndex;
-				const FTrait* Trait = TraitRegistry.Find(TraitTemplate->GetRegistryHandle());
-				const FString TraitName = Trait != nullptr ? Trait->GetTraitName() : TEXT("<Unknown>");
-
-				const uint32 NextTraitIndex = TraitIndex + 1;
-				const uint32 EndOfNextTraitSharedData = NextTraitIndex < NumTraits ? TraitTemplates[NextTraitIndex].GetNodeSharedOffset() : NodeTemplate->GetNodeSharedDataSize();
-				const uint32 TraitSharedDataSize = EndOfNextTraitSharedData - TraitTemplate->GetNodeSharedOffset();
-
-				const uint32 EndOfNextTraitInstanceData = NextTraitIndex < NumTraits ? TraitTemplates[NextTraitIndex].GetNodeInstanceOffset() : NodeTemplate->GetNodeInstanceDataSize();
-				const uint32 TraitInstanceDataSize = EndOfNextTraitInstanceData - TraitTemplate->GetNodeInstanceOffset();
-
-				UE_LOG(LogAnimation, Log, TEXT("            %u: [%x] %s (%s)"), TraitIndex, TraitTemplate->GetUID().GetUID(), *TraitName, TraitTemplate->GetMode() == ETraitMode::Base ? TEXT("Base") : TEXT("Additive"));
-				UE_LOG(LogAnimation, Log, TEXT("                Shared Data: [Offset: %u bytes, Size: %u bytes]"), TraitTemplate->GetNodeSharedOffset(), TraitSharedDataSize);
-				if (TraitTemplate->HasLatentProperties() && Trait != nullptr)
-				{
-					UE_LOG(LogAnimation, Log, TEXT("                Shared Data Latent Property Handles: [Offset: %u bytes, Count: %u]"), TraitTemplate->GetNodeSharedLatentPropertyHandlesOffset(), Trait->GetNumLatentTraitProperties());
-				}
-				UE_LOG(LogAnimation, Log, TEXT("                Instance Data: [Offset: %u bytes, Size: %u bytes]"), TraitTemplate->GetNodeInstanceOffset(), TraitInstanceDataSize);
-			}
-		}
-
-		LogAnimation.SetVerbosity(OldVerbosity);
 	}
-
-	void FAnimNextModuleImpl::ListAnimationGraphs(const TArray<FString>& Args)
-	{
-		// Turn off log times to make diff-ing easier
-		TGuardValue<ELogTimes::Type> DisableLogTimes(GPrintLogTimes, ELogTimes::None);
-
-		// Make sure to log everything
-		const ELogVerbosity::Type OldVerbosity = LogAnimation.GetVerbosity();
-		LogAnimation.SetVerbosity(ELogVerbosity::All);
-
-		TArray<const UAnimNextAnimationGraph*> AnimationGraphs;
-
-		for (TObjectIterator<UAnimNextAnimationGraph> It; It; ++It)
-		{
-			AnimationGraphs.Add(*It);
-		}
-
-		struct FCompareObjectNames
-		{
-			FORCEINLINE bool operator()(const UAnimNextAnimationGraph& Lhs, const UAnimNextAnimationGraph& Rhs) const
-			{
-				return Lhs.GetPathName().Compare(Rhs.GetPathName()) < 0;
-			}
-		};
-		AnimationGraphs.Sort(FCompareObjectNames());
-
-		const FNodeTemplateRegistry& NodeTemplateRegistry = FNodeTemplateRegistry::Get();
-		const FTraitRegistry& TraitRegistry = FTraitRegistry::Get();
-		const bool bDetailedOutput = true;
-
-		UE_LOG(LogAnimation, Log, TEXT("===== AnimNext Modules ====="));
-		UE_LOG(LogAnimation, Log, TEXT("Num Graphs: %u"), AnimationGraphs.Num());
-
-		for (const UAnimNextAnimationGraph* AnimationGraph : AnimationGraphs)
-		{
-			uint32 TotalInstanceSize = 0;
-			uint32 NumNodes = 0;
-			{
-				// We always have a node at offset 0
-				int32 NodeOffset = 0;
-
-				while (NodeOffset < AnimationGraph->SharedDataBuffer.Num())
-				{
-					const FNodeDescription* NodeDesc = reinterpret_cast<const FNodeDescription*>(&AnimationGraph->SharedDataBuffer[NodeOffset]);
-
-					TotalInstanceSize += NodeDesc->GetNodeInstanceDataSize();
-					NumNodes++;
-
-					const FNodeTemplate* NodeTemplate = NodeTemplateRegistry.Find(NodeDesc->GetTemplateHandle());
-					NodeOffset += NodeTemplate->GetNodeSharedDataSize();
-				}
-			}
-
-			UE_LOG(LogAnimation, Log, TEXT("    %s ..."), *AnimationGraph->GetPathName());
-			UE_LOG(LogAnimation, Log, TEXT("        Shared Data Size: %.2f KB"), double(AnimationGraph->SharedDataBuffer.Num()) / 1024.0);
-			UE_LOG(LogAnimation, Log, TEXT("        Max Instance Data Size: %.2f KB"), double(TotalInstanceSize) / 1024.0);
-			UE_LOG(LogAnimation, Log, TEXT("        Num Nodes: %u"), NumNodes);
-
-			if (bDetailedOutput)
-			{
-				// We always have a node at offset 0
-				int32 NodeOffset = 0;
-
-				while (NodeOffset < AnimationGraph->SharedDataBuffer.Num())
-				{
-					const FNodeDescription* NodeDesc = reinterpret_cast<const FNodeDescription*>(&AnimationGraph->SharedDataBuffer[NodeOffset]);
-					const FNodeTemplate* NodeTemplate = NodeTemplateRegistry.Find(NodeDesc->GetTemplateHandle());
-
-					const uint32 NumTraits = NodeTemplate->GetNumTraits();
-
-					UE_LOG(LogAnimation, Log, TEXT("        Node %u: [Template %x with %u traits]"), NodeDesc->GetUID().GetNodeIndex(), NodeTemplate->GetUID(), NumTraits);
-					UE_LOG(LogAnimation, Log, TEXT("            Shared Data: [Offset: %u bytes, Size: %u bytes]"), NodeOffset, NodeTemplate->GetNodeSharedDataSize());
-					UE_LOG(LogAnimation, Log, TEXT("            Instance Data Size: %u bytes"), NodeDesc->GetNodeInstanceDataSize());
-					UE_LOG(LogAnimation, Log, TEXT("            Traits ..."));
-
-					const FTraitTemplate* TraitTemplates = NodeTemplate->GetTraits();
-					for (uint32 TraitIndex = 0; TraitIndex < NumTraits; ++TraitIndex)
-					{
-						const FTraitTemplate* TraitTemplate = TraitTemplates + TraitIndex;
-						const FTrait* Trait = TraitRegistry.Find(TraitTemplate->GetRegistryHandle());
-						const FString TraitName = Trait != nullptr ? Trait->GetTraitName() : TEXT("<Unknown>");
-
-						const uint32 NextTraitIndex = TraitIndex + 1;
-						const uint32 EndOfNextTraitSharedData = NextTraitIndex < NumTraits ? TraitTemplates[NextTraitIndex].GetNodeSharedOffset() : NodeTemplate->GetNodeSharedDataSize();
-						const uint32 TraitSharedDataSize = EndOfNextTraitSharedData - TraitTemplate->GetNodeSharedOffset();
-
-						const uint32 EndOfNextTraitInstanceData = NextTraitIndex < NumTraits ? TraitTemplates[NextTraitIndex].GetNodeInstanceOffset() : NodeTemplate->GetNodeInstanceDataSize();
-						const uint32 TraitInstanceDataSize = EndOfNextTraitInstanceData - TraitTemplate->GetNodeInstanceOffset();
-
-						UE_LOG(LogAnimation, Log, TEXT("                    %u: [%x] %s (%s)"), TraitIndex, TraitTemplate->GetUID().GetUID(), *TraitName, TraitTemplate->GetMode() == ETraitMode::Base ? TEXT("Base") : TEXT("Additive"));
-						UE_LOG(LogAnimation, Log, TEXT("                        Shared Data: [Offset: %u bytes, Size: %u bytes]"), TraitTemplate->GetNodeSharedOffset(), TraitSharedDataSize);
-						if (TraitTemplate->HasLatentProperties() && Trait != nullptr)
-						{
-							UE_LOG(LogAnimation, Log, TEXT("                        Shared Data Latent Property Handles: [Offset: %u bytes, Count: %u]"), TraitTemplate->GetNodeSharedLatentPropertyHandlesOffset(), Trait->GetNumLatentTraitProperties());
-						}
-						UE_LOG(LogAnimation, Log, TEXT("                        Instance Data: [Offset: %u bytes, Size: %u bytes]"), TraitTemplate->GetNodeInstanceOffset(), TraitInstanceDataSize);
-					}
-
-					NodeOffset += NodeTemplate->GetNodeSharedDataSize();
-				}
-			}
-		}
-
-		LogAnimation.SetVerbosity(OldVerbosity);
-	}
-#endif
-
-	IAnimNextModuleInterface& IAnimNextModuleInterface::Get()
-	{
-		return FModuleManager::LoadModuleChecked<IAnimNextModuleInterface>(TEXT("AnimNext"));
-	}
-
 }
 
 #undef LOCTEXT_NAMESPACE
diff --git a/Engine/Plugins/Experimental/AnimNext/Source/AnimNext/Private/Graph/RigUnit_AnimNextRemapPose.h b/Engine/Plugins/Experimental/AnimNext/Source/AnimNext/Private/Graph/RigUnit_AnimNextRemapPose.h
--- a/Engine/Plugins/Experimental/AnimNext/Source/AnimNext/Private/Graph/RigUnit_AnimNextRemapPose.h
+++ b/Engine/Plugins/Experimental/AnimNext/Source/AnimNext/Private/Graph/RigUnit_AnimNextRemapPose.h
@@ -11,7 +11,7 @@
 
 /** Remaps an anim graph pose from to another skeletal mesh component. */
 USTRUCT(meta=(DisplayName="Remap Pose", Category="Animation Graph", NodeColor="0, 1, 1", Keywords="Output,Pose,Port"))
-struct ANIMNEXT_API FRigUnit_AnimNextRemapPose : public FRigUnit_AnimNextBase
+struct FRigUnit_AnimNextRemapPose : public FRigUnit_AnimNextBase
 {
 	GENERATED_BODY()
 
diff --git a/Engine/Plugins/Experimental/AnimNext/Source/AnimNext/Private/Graph/RigUnit_AnimNextWriteSkeletalMeshComponentPose.h b/Engine/Plugins/Experimental/AnimNext/Source/AnimNext/Private/Graph/RigUnit_AnimNextWriteSkeletalMeshComponentPose.h
--- a/Engine/Plugins/Experimental/AnimNext/Source/AnimNext/Private/Graph/RigUnit_AnimNextWriteSkeletalMeshComponentPose.h
+++ b/Engine/Plugins/Experimental/AnimNext/Source/AnimNext/Private/Graph/RigUnit_AnimNextWriteSkeletalMeshComponentPose.h
@@ -10,7 +10,7 @@
 
 /** Writes a pose to a skeletal mesh component */
 USTRUCT(meta=(DisplayName="Write Pose", Category="Animation Graph", NodeColor="0, 1, 1", Keywords="Output,Pose,Port"))
-struct ANIMNEXT_API FRigUnit_AnimNextWriteSkeletalMeshComponentPose : public FRigUnit_AnimNextBase
+struct FRigUnit_AnimNextWriteSkeletalMeshComponentPose : public FRigUnit_AnimNextBase
 {
 	GENERATED_BODY()
 
diff --git a/Engine/Plugins/Experimental/AnimNext/Source/AnimNext/Private/Graph/RigUnit_GetActorTransform.h b/Engine/Plugins/Experimental/AnimNext/Source/AnimNext/Private/Graph/RigUnit_GetActorTransform.h
--- a/Engine/Plugins/Experimental/AnimNext/Source/AnimNext/Private/Graph/RigUnit_GetActorTransform.h
+++ b/Engine/Plugins/Experimental/AnimNext/Source/AnimNext/Private/Graph/RigUnit_GetActorTransform.h
@@ -9,7 +9,7 @@
 
 /** Gets the transform of the host of this AnimNext module */
 USTRUCT(meta=(DisplayName="Get Host Transform", Category="Animation Graph", NodeColor="0, 1, 1", Keywords="Actor, Transform"))
-struct ANIMNEXT_API FRigUnit_GetActorTransform : public FRigUnit_AnimNextBase
+struct FRigUnit_GetActorTransform : public FRigUnit_AnimNextBase
 {
 	GENERATED_BODY()
 
diff --git a/Engine/Plugins/Experimental/AnimNext/Source/AnimNext/Private/Graph/RigUnit_MakeReferencePoseFromSkeletalMeshComponent.h b/Engine/Plugins/Experimental/AnimNext/Source/AnimNext/Private/Graph/RigUnit_MakeReferencePoseFromSkeletalMeshComponent.h
--- a/Engine/Plugins/Experimental/AnimNext/Source/AnimNext/Private/Graph/RigUnit_MakeReferencePoseFromSkeletalMeshComponent.h
+++ b/Engine/Plugins/Experimental/AnimNext/Source/AnimNext/Private/Graph/RigUnit_MakeReferencePoseFromSkeletalMeshComponent.h
@@ -10,7 +10,7 @@
 
 /** Makes a reference pose from a skeletal mesh component */
 USTRUCT(meta=(DisplayName="Make Reference Pose", Category="Animation Graph", NodeColor="0, 1, 1", Keywords="Output,Pose,Port"))
-struct ANIMNEXT_API FRigUnit_MakeReferencePoseFromSkeletalMeshComponent : public FRigUnit_AnimNextBase
+struct FRigUnit_MakeReferencePoseFromSkeletalMeshComponent : public FRigUnit_AnimNextBase
 {
 	GENERATED_BODY()
 
diff --git a/Engine/Plugins/Experimental/AnimNext/Source/AnimNext/Private/Module/AnimNextModule.cpp b/Engine/Plugins/Experimental/AnimNext/Source/AnimNext/Private/Module/AnimNextModule.cpp
--- a/Engine/Plugins/Experimental/AnimNext/Source/AnimNext/Private/Module/AnimNextModule.cpp
+++ b/Engine/Plugins/Experimental/AnimNext/Source/AnimNext/Private/Module/AnimNextModule.cpp
@@ -2,13 +2,9 @@
 
 #include "Module/AnimNextModule.h"
 #include "RigVMCore/RigVMExecuteContext.h"
-#include "Graph/RigUnit_AnimNextShimRoot.h"
 #include "UObject/AssetRegistryTagsContext.h"
 #include "UObject/ObjectSaveContext.h"
-#include "TraitCore/ExecutionContext.h"
-#include "Graph/AnimNext_LODPose.h"
 #include "Serialization/MemoryReader.h"
-#include "Graph/AnimNextGraphEntryPoint.h"
 #include "Graph/RigUnit_AnimNextBeginExecution.h"
 #include "UObject/LinkerLoad.h"
 #include "UObject/ObjectResource.h"
diff --git a/Engine/Plugins/Experimental/AnimNext/Source/AnimNext/Private/Module/ModuleTickFunction.cpp b/Engine/Plugins/Experimental/AnimNext/Source/AnimNext/Private/Module/ModuleTickFunction.cpp
--- a/Engine/Plugins/Experimental/AnimNext/Source/AnimNext/Private/Module/ModuleTickFunction.cpp
+++ b/Engine/Plugins/Experimental/AnimNext/Source/AnimNext/Private/Module/ModuleTickFunction.cpp
@@ -135,14 +135,12 @@ void FModuleEventTickFunction::Run(float InDeltaTime)
 		AnimNextContext.SetDrawInterface(&ModuleInstance->GetDebugDrawInterface());
 #endif
 
-		// Module setup
-		AnimNextContext.SetContextData<FAnimNextModuleContextData>(ModuleInstance);
+		// Insert our context data for the scope of execution
+		FAnimNextModuleContextData ContextData(ModuleInstance);
+		FScopedExecuteContextData ContextDataScope(AnimNextContext, ContextData);
 
 		// Run the VM for this event
 		VM->ExecuteVM(Context, EventName);
-
-		// Reset the context to avoid issues if we forget to reset it the next time we use it
-		AnimNextContext.DebugReset<FAnimNextModuleContextData>();
 	}
 
 	while (!PostExecuteTasks.IsEmpty())
diff --git a/Engine/Plugins/Experimental/AnimNext/Source/AnimNext/Private/Param/ParamType.cpp b/Engine/Plugins/Experimental/AnimNext/Source/AnimNext/Private/Param/ParamType.cpp
--- a/Engine/Plugins/Experimental/AnimNext/Source/AnimNext/Private/Param/ParamType.cpp
+++ b/Engine/Plugins/Experimental/AnimNext/Source/AnimNext/Private/Param/ParamType.cpp
@@ -2,14 +2,11 @@
 
 #include "Param/ParamType.h"
 
-#include "Graph/AnimNext_LODPose.h"
 #include "Misc/StringBuilder.h"
 #include "UObject/Class.h"
 #include "Templates/SubclassOf.h"
-#include "GameFramework/CharacterMovementComponent.h"
-#include "Components/SkeletalMeshComponent.h"
-#include "Animation/AnimSequence.h"
 #include "RigVMCore/RigVMTemplate.h"
+#include "UObject/TextProperty.h"
 
 FAnimNextParamType::FAnimNextParamType(EValueType InValueType, EContainerType InContainerType, const UObject* InValueTypeObject)
 	: ValueTypeObject(InValueTypeObject)
diff --git a/Engine/Plugins/Experimental/AnimNext/Source/AnimNext/Private/TraitCore/TraitEvent.cpp b/Engine/Plugins/Experimental/AnimNext/Source/AnimNext/Private/TraitCore/TraitEvent.cpp
--- a/Engine/Plugins/Experimental/AnimNext/Source/AnimNext/Private/TraitCore/TraitEvent.cpp
+++ b/Engine/Plugins/Experimental/AnimNext/Source/AnimNext/Private/TraitCore/TraitEvent.cpp
@@ -1,10 +1,6 @@
 // Copyright Epic Games, Inc. All Rights Reserved.
 
 #include "TraitCore/TraitEvent.h"
-#include "TraitCore/TraitEventList.h"
-#include "TraitCore/TraitBinding.h"
-#include "TraitCore/TraitStackBinding.h"
-#include "Graph/GraphInstanceComponent.h"
 
 bool FAnimNextTraitEvent::DecrementLifetime(UE::AnimNext::FTraitEventList& OutputEventList)
 {
@@ -16,58 +12,3 @@ bool FAnimNextTraitEvent::DecrementLifetime(UE::AnimNext::FTraitEventList& Outpu
 
 	return bExpired;
 }
-
-namespace UE::AnimNext
-{
-	void RaiseTraitEvent(FExecutionContext& Context, const FTraitStackBinding& StackBinding, FAnimNextTraitEvent& Event)
-	{
-		if (!StackBinding.IsValid() || !Event.IsValid())
-		{
-			return;
-		}
-
-		FTraitBinding Binding;
-		ensure(StackBinding.GetTopTrait(Binding));
-		do
-		{
-			const ETraitStackPropagation Propagation = Binding.GetTrait()->OnTraitEvent(Context, Binding, Event);
-			if (Propagation == ETraitStackPropagation::Stop || Event.IsConsumed())
-			{
-				// We don't wish to propagate the call further, stop here
-				break;
-			}
-		} while (StackBinding.GetParentTrait(Binding, Binding));
-	}
-
-	void RaiseTraitEvents(FExecutionContext& Context, const FTraitStackBinding& StackBinding, const FTraitEventList& EventList)
-	{
-		if (!StackBinding.IsValid())
-		{
-			return;
-		}
-
-		// Event handlers can raise events and as such the list may change while we iterate
-		// However, if an event is added while we iterate, we will not visit it
-		const int32 NumEvents = EventList.Num();
-		for (int32 EventIndex = 0; EventIndex < NumEvents; ++EventIndex)
-		{
-			const FAnimNextTraitEventPtr Event = EventList[EventIndex];
-			RaiseTraitEvent(Context, StackBinding, *Event);
-		}
-	}
-
-	void RaiseTraitEvents(FExecutionContext& Context, FGraphInstanceComponent& Component, const FTraitEventList& EventList)
-	{
-		// Event handlers can raise events and as such the list may change while we iterate
-		// However, if an event is added while we iterate, we will not visit it
-		const int32 NumEvents = EventList.Num();
-		for (int32 EventIndex = 0; EventIndex < NumEvents; ++EventIndex)
-		{
-			const FAnimNextTraitEventPtr Event = EventList[EventIndex];
-			if (Event->IsValid())
-			{
-				Component.OnTraitEvent(Context, *Event);
-			}
-		}
-	}
-}
diff --git a/Engine/Plugins/Experimental/AnimNext/Source/AnimNext/Public/AnimNextStats.h b/Engine/Plugins/Experimental/AnimNext/Source/AnimNext/Public/AnimNextStats.h
--- a/Engine/Plugins/Experimental/AnimNext/Source/AnimNext/Public/AnimNextStats.h
+++ b/Engine/Plugins/Experimental/AnimNext/Source/AnimNext/Public/AnimNextStats.h
@@ -13,29 +13,16 @@ DECLARE_CYCLE_STAT_EXTERN(TEXT("AnimNext: Initialize Instance"), STAT_AnimNext_I
 DECLARE_CYCLE_STAT_EXTERN(TEXT("AnimNext: Create Instance Data"), STAT_AnimNext_CreateInstanceData, STATGROUP_AnimNext, ANIMNEXT_API);
 
 DECLARE_CYCLE_STAT_EXTERN(TEXT("AnimNext: Make RefPose"), STAT_AnimNext_Make_RefPose, STATGROUP_AnimNext, ANIMNEXT_API);
-DECLARE_CYCLE_STAT_EXTERN(TEXT("AnimNext: Run Graph"), STAT_AnimNext_Run_Graph, STATGROUP_AnimNext, ANIMNEXT_API);
 DECLARE_CYCLE_STAT_EXTERN(TEXT("AnimNext: Write Pose"), STAT_AnimNext_Write_Pose, STATGROUP_AnimNext, ANIMNEXT_API);
 
-DECLARE_CYCLE_STAT_EXTERN(TEXT("AnimNext: Allocate Graph Instance"), STAT_AnimNext_Graph_AllocateInstance, STATGROUP_AnimNext, ANIMNEXT_API);
-DECLARE_CYCLE_STAT_EXTERN(TEXT("AnimNext: Graph RigVM"), STAT_AnimNext_Graph_RigVM, STATGROUP_AnimNext, ANIMNEXT_API);
-DECLARE_CYCLE_STAT_EXTERN(TEXT("AnimNext: Execute Evaluation Program"), STAT_AnimNext_EvaluationProgram_Execute, STATGROUP_AnimNext, ANIMNEXT_API);
 DECLARE_CYCLE_STAT_EXTERN(TEXT("AnimNext: Copy Transforms (SoA)"), STAT_AnimNext_CopyTransforms_SoA, STATGROUP_AnimNext, ANIMNEXT_API);
 DECLARE_CYCLE_STAT_EXTERN(TEXT("AnimNext: Normalize Rotations (SoA)"), STAT_AnimNext_NormalizeRotations_SoA, STATGROUP_AnimNext, ANIMNEXT_API);
 DECLARE_CYCLE_STAT_EXTERN(TEXT("AnimNext: Blend Overwrite (SoA)"), STAT_AnimNext_BlendOverwrite_SoA, STATGROUP_AnimNext, ANIMNEXT_API);
 DECLARE_CYCLE_STAT_EXTERN(TEXT("AnimNext: Blend Accumulate (SoA)"), STAT_AnimNext_BlendAccumulate_SoA, STATGROUP_AnimNext, ANIMNEXT_API);
 
-DECLARE_CYCLE_STAT_EXTERN(TEXT("AnimNext: Param Type Handle Lock"), STAT_AnimNext_ParamTypeHandle_Lock, STATGROUP_AnimNext, ANIMNEXT_API);
-DECLARE_CYCLE_STAT_EXTERN(TEXT("AnimNext: Param Stack Get Param"), STAT_AnimNext_ParamStack_GetParam, STATGROUP_AnimNext, ANIMNEXT_API);
-DECLARE_CYCLE_STAT_EXTERN(TEXT("AnimNext: Param Stack Adapter"), STAT_AnimNext_ParamStack_Adapter, STATGROUP_AnimNext, ANIMNEXT_API);
-DECLARE_CYCLE_STAT_EXTERN(TEXT("AnimNext: Param Stack Coalesce"), STAT_AnimNext_ParamStack_Coalesce, STATGROUP_AnimNext, ANIMNEXT_API);
-DECLARE_CYCLE_STAT_EXTERN(TEXT("AnimNext: Param Stack Decoalesce"), STAT_AnimNext_ParamStack_Decoalesce, STATGROUP_AnimNext, ANIMNEXT_API);
-DECLARE_CYCLE_STAT_EXTERN(TEXT("AnimNext: Update Param Layer"), STAT_AnimNext_Graph_UpdateParamLayer, STATGROUP_AnimNext, ANIMNEXT_API);
-
 DECLARE_CYCLE_STAT_EXTERN(TEXT("AnimNext: Generate Reference Pose"), STAT_AnimNext_GenerateReferencePose, STATGROUP_AnimNext, ANIMNEXT_API);
 DECLARE_CYCLE_STAT_EXTERN(TEXT("AnimNext: Remap Pose (From AnimBP)"), STAT_AnimNext_RemapPose_FromAnimBP, STATGROUP_AnimNext, ANIMNEXT_API);
 DECLARE_CYCLE_STAT_EXTERN(TEXT("AnimNext: Remap Pose (To AnimBP)"), STAT_AnimNext_RemapPose_ToAnimBP, STATGROUP_AnimNext, ANIMNEXT_API);
 DECLARE_CYCLE_STAT_EXTERN(TEXT("AnimNext: Remap Pose (To LocalTransforms)"), STAT_AnimNext_RemapPose_ToLocalTransforms, STATGROUP_AnimNext, ANIMNEXT_API);
 DECLARE_CYCLE_STAT_EXTERN(TEXT("AnimNext: Remap Pose (Mesh->Mesh)"), STAT_AnimNext_RemapPose_Mesh2Mesh, STATGROUP_AnimNext, ANIMNEXT_API);
-DECLARE_CYCLE_STAT_EXTERN(TEXT("AnimNext: Convert Local Space To Component Space"), STAT_AnimNext_ConvertLocalSpaceToComponentSpace, STATGROUP_AnimNext, ANIMNEXT_API);
-
-DECLARE_CYCLE_STAT_EXTERN(TEXT("AnimNext: Skeletal Mesh Component Port"), STAT_AnimNext_Port_SkeletalMeshComponent, STATGROUP_AnimNext, ANIMNEXT_API);
+DECLARE_CYCLE_STAT_EXTERN(TEXT("AnimNext: Convert Local Space To Component Space"), STAT_AnimNext_ConvertLocalSpaceToComponentSpace, STATGROUP_AnimNext, ANIMNEXT_API);
\ No newline at end of file
diff --git a/Engine/Plugins/Experimental/AnimNext/Source/AnimNext/Public/TraitCore/TraitEvent.h b/Engine/Plugins/Experimental/AnimNext/Source/AnimNext/Public/TraitCore/TraitEvent.h
--- a/Engine/Plugins/Experimental/AnimNext/Source/AnimNext/Public/TraitCore/TraitEvent.h
+++ b/Engine/Plugins/Experimental/AnimNext/Source/AnimNext/Public/TraitCore/TraitEvent.h
@@ -130,22 +130,3 @@ inline TSharedPtr<EventType, ESPMode::ThreadSafe> MakeTraitEvent(TArgs&&... Args
 {
 	return MakeShared<EventType>(Forward<TArgs>(Args)...);
 }
-
-namespace UE::AnimNext
-{
-	struct FExecutionContext;
-	struct FTraitStackBinding;
-	struct FTraitEventList;
-	struct FGraphInstanceComponent;
-
-	// Raises a trait event on the specified trait stack
-	// This calls OnTraitEvent on each trait of the stack following propagation rules
-	ANIMNEXT_API void RaiseTraitEvent(FExecutionContext& Context, const FTraitStackBinding& StackBinding, FAnimNextTraitEvent& Event);
-
-	// Raises every trait event from the input list on the specified trait stack
-	// This calls OnTraitEvent on each trait of the stack following propagation rules
-	ANIMNEXT_API void RaiseTraitEvents(FExecutionContext& Context, const FTraitStackBinding& StackBinding, const FTraitEventList& EventList);
-
-	// Raises every trait event from the input list on the specified component
-	ANIMNEXT_API void RaiseTraitEvents(FExecutionContext& Context, FGraphInstanceComponent& Component, const FTraitEventList& EventList);
-}
diff --git a/Engine/Plugins/Experimental/AnimNext/Source/AnimNextEditor/Internal/Common/AnimNextAssetItemDetails.cpp b/Engine/Plugins/Experimental/AnimNext/Source/AnimNextEditor/Internal/Common/AnimNextAssetItemDetails.cpp
new file mode 100644
--- /dev/null
+++ b/Engine/Plugins/Experimental/AnimNext/Source/AnimNextEditor/Internal/Common/AnimNextAssetItemDetails.cpp
@@ -0,0 +1,212 @@
+// Copyright Epic Games, Inc. All Rights Reserved.
+
+#include "AnimNextAssetItemDetails.h"
+
+#include "ClassIconFinder.h"
+#include "Toolkits/AssetEditorToolkitMenuContext.h"
+#include "Entries/AnimNextEventGraphEntry.h"
+#include "InstancedStruct.h"
+#include "AnimNextAssetWorkspaceAssetUserData.h"
+#include "RigVMModel/RigVMGraph.h"
+#include "WorkspaceItemMenuContext.h"
+#include "IWorkspaceEditor.h"
+#include "ScopedTransaction.h"
+#include "RigVMModel/RigVMClient.h"
+#include "ToolMenus.h"
+#include "Entries/AnimNextDataInterfaceEntry.h"
+#include "Framework/Commands/GenericCommands.h"
+#include "Module/AnimNextModule_EditorData.h"
+#include "Module/RigUnit_AnimNextModuleEvents.h"
+#include "UObject/UObjectIterator.h"
+#include "Variables/SAddVariablesDialog.h"
+
+#define LOCTEXT_NAMESPACE "FAnimNextGraphItemDetails"
+
+namespace UE::AnimNext::Editor
+{
+
+const FSlateBrush* FAnimNextAssetItemDetails::GetItemIcon(const FWorkspaceOutlinerItemExport& Export) const
+{
+	return FAppStyle::GetBrush(TEXT("LevelEditor.Tabs.Outliner"));
+}
+
+void FAnimNextAssetItemDetails::RegisterToolMenuExtensions()
+{
+	FToolMenuOwnerScoped OwnerScoped(TEXT("FAnimNextModuleItemDetails"));
+	UToolMenu* Menu = UToolMenus::Get()->ExtendMenu("WorkspaceOutliner.ItemContextMenu");
+	if (Menu == nullptr)
+	{
+		return;
+	}
+
+	Menu->AddDynamicSection(TEXT("AnimNextAssetItem"), FNewToolMenuDelegate::CreateLambda([](UToolMenu* InMenu)
+	{
+		UWorkspaceItemMenuContext* WorkspaceItemContext = InMenu->FindContext<UWorkspaceItemMenuContext>();
+		const UAssetEditorToolkitMenuContext* AssetEditorContext = InMenu->FindContext<UAssetEditorToolkitMenuContext>();
+		if (WorkspaceItemContext == nullptr || AssetEditorContext == nullptr)
+		{
+			return;
+		}
+
+		const TSharedPtr<UE::Workspace::IWorkspaceEditor> WorkspaceEditor = StaticCastSharedPtr<UE::Workspace::IWorkspaceEditor>(AssetEditorContext->Toolkit.Pin());
+		if(!WorkspaceEditor.IsValid())
+		{
+			return;
+		}
+
+		if (WorkspaceItemContext->SelectedExports.Num() != 1)
+		{
+			return;
+		}
+
+		TInstancedStruct<FWorkspaceOutlinerItemData>& Data = WorkspaceItemContext->SelectedExports[0].GetData();
+		if (!Data.IsValid() || !Data.GetScriptStruct()->IsChildOf(FAnimNextRigVMAssetOutlinerData::StaticStruct()))
+		{
+			return;
+		}
+
+		const FAnimNextRigVMAssetOutlinerData& OutlinerData = Data.Get<FAnimNextRigVMAssetOutlinerData>();
+		UAnimNextRigVMAsset* Asset = OutlinerData.Asset;
+		if(Asset == nullptr)
+		{
+			return;
+		}
+
+		UAnimNextRigVMAssetEditorData* EditorData = UncookedOnly::FUtils::GetEditorData<UAnimNextRigVMAssetEditorData>(Asset);
+		if(EditorData == nullptr)
+		{
+			return;
+		}
+
+		FToolMenuSection& AssetSection = InMenu->AddSection("AnimNextAsset", LOCTEXT("AnimNextAssetSectionLabel", "AnimNext Asset"));
+
+		// Per sub-item type addition
+		TConstArrayView<TSubclassOf<UAnimNextRigVMAssetEntry>> AllEntryClasses = EditorData->GetEntryClasses();
+		for(UClass* SubEntryClass : AllEntryClasses)
+		{
+			static FTextFormat AddEntryLabelFormat(LOCTEXT("AddEntryLabelFormat", "Add {0}"));
+			static FTextFormat AddEntryTooltipFormat(LOCTEXT("AddEntryTooltipFormat", "Adds a new {0} to this asset"));
+
+			if(!EditorData->CanAddNewEntry(SubEntryClass))
+			{
+				continue;
+			}
+
+			if(SubEntryClass == UAnimNextEventGraphEntry::StaticClass())
+			{
+				AssetSection.AddSubMenu(
+					SubEntryClass->GetFName(),
+					FText::Format(AddEntryLabelFormat, SubEntryClass->GetDisplayNameText()),
+					FText::Format(AddEntryTooltipFormat, SubEntryClass->GetDisplayNameText()),
+					FNewToolMenuDelegate::CreateLambda([SubEntryClass, WorkspaceItemContext, EditorData, Asset](UToolMenu* InToolMenu)
+					{
+						FToolMenuSection& Section = InToolMenu->AddSection(SubEntryClass->GetFName());
+
+						static FTextFormat AddEventTooltipFormat(LOCTEXT("AddEventGraphTooltipFormat", "Adds a {0} event graph to this asset"));
+						
+						for(TObjectIterator<UScriptStruct> It; It; ++It)
+						{
+							UScriptStruct* Struct = *It;
+							if(!Struct->IsChildOf(FRigUnit_AnimNextModuleEventBase::StaticStruct()) || Struct == FRigUnit_AnimNextModuleEventBase::StaticStruct())
+							{
+								continue;
+							}
+							if(Struct->HasMetaData(FRigVMStruct::HiddenMetaName) || Struct->HasMetaData(FRigVMStruct::AbstractMetaName))
+							{
+								continue;
+							}
+							
+							TInstancedStruct<FRigUnit_AnimNextModuleEventBase> StructInstance;
+							StructInstance.InitializeAsScriptStruct(Struct);
+							FName EventName = StructInstance.Get<FRigUnit_AnimNextModuleEventBase>().GetEventName();
+
+							Section.AddMenuEntry(
+								StructInstance.Get<FRigUnit_AnimNextModuleEventBase>().GetEventName(),
+								FText::FromName(EventName),
+								FText::Format(AddEventTooltipFormat, FText::FromName(EventName)),
+								FSlateIconFinder::FindIconForClass(SubEntryClass, "ClassIcon.Object"),
+								FUIAction(
+									FExecuteAction::CreateWeakLambda(WorkspaceItemContext, [EventName, SubEntryClass, EditorData, Struct]()
+									{
+										FScopedTransaction Transaction(LOCTEXT("AddEventGraph", "Add Event Graph"));
+										EditorData->AddEventGraph(EventName, Struct);
+									}),
+									FCanExecuteAction::CreateWeakLambda(WorkspaceItemContext, [Asset, EventName]()
+									{
+										return !Asset->GetVM()->ContainsEntry(EventName);
+									})
+								));
+						}
+					}),
+					false,
+					FSlateIconFinder::FindIconForClass(SubEntryClass, "ClassIcon.Object"));
+			}
+			else
+			{
+				AssetSection.AddMenuEntry(
+					SubEntryClass->GetFName(),
+					FText::Format(AddEntryLabelFormat, SubEntryClass->GetDisplayNameText()),
+					FText::Format(AddEntryTooltipFormat, SubEntryClass->GetDisplayNameText()),
+					FSlateIconFinder::FindIconForClass(SubEntryClass, "ClassIcon.Object"),
+					FUIAction(
+						FExecuteAction::CreateWeakLambda(WorkspaceItemContext, [&OutlinerData, SubEntryClass, EditorData]()
+						{
+							if(SubEntryClass == UAnimNextVariableEntry::StaticClass())
+							{
+								TSharedRef<SAddVariablesDialog> AddVariablesDialog =
+									SNew(SAddVariablesDialog, TArray<UAnimNextRigVMAssetEditorData*>({ UncookedOnly::FUtils::GetEditorData<UAnimNextRigVMAssetEditorData>(OutlinerData.Asset.Get()) }));
+
+								TArray<FVariableToAdd> VariablesToAdd;
+								TArray<FDataInterfaceToAdd> DataInterfacesToAdd;
+								if(AddVariablesDialog->ShowModal(VariablesToAdd, DataInterfacesToAdd))
+								{
+									FScopedTransaction Transaction(LOCTEXT("AddVariables", "Add Variable(s)"));
+									for (const FVariableToAdd& VariableToAdd : VariablesToAdd)
+									{
+										EditorData->AddVariable(VariableToAdd.Name, VariableToAdd.Type);
+									}
+									for (const FDataInterfaceToAdd& DataInterfaceToAdd : DataInterfacesToAdd)
+									{
+										EditorData->AddDataInterface(DataInterfaceToAdd.DataInterface);
+									}
+								}
+							}
+							else if(SubEntryClass == UAnimNextDataInterfaceEntry::StaticClass())
+							{
+								TSharedRef<SAddVariablesDialog> AddVariablesDialog =
+									SNew(SAddVariablesDialog, TArray<UAnimNextRigVMAssetEditorData*>({ UncookedOnly::FUtils::GetEditorData<UAnimNextRigVMAssetEditorData>(OutlinerData.Asset.Get()) }))
+									.ShouldAddInitialVariable(false);
+
+								TArray<FVariableToAdd> VariablesToAdd;
+								TArray<FDataInterfaceToAdd> DataInterfacesToAdd;
+								if(AddVariablesDialog->ShowModal(VariablesToAdd, DataInterfacesToAdd))
+								{
+									FScopedTransaction Transaction(LOCTEXT("AddVariables", "Add Variable(s)"));
+									for (const FVariableToAdd& VariableToAdd : VariablesToAdd)
+									{
+										EditorData->AddVariable(VariableToAdd.Name, VariableToAdd.Type);
+									}
+									for (const FDataInterfaceToAdd& DataInterfaceToAdd : DataInterfacesToAdd)
+									{
+										EditorData->AddDataInterface(DataInterfaceToAdd.DataInterface);
+									}
+								}
+							}
+						})
+					));
+			}
+		}
+	}));
+}
+
+void FAnimNextAssetItemDetails::UnregisterToolMenuExtensions()
+{
+	if(UToolMenus* ToolMenus = UToolMenus::Get())
+	{
+		ToolMenus->UnregisterOwnerByName("FAnimNextAssetItemDetails");
+	}
+}
+
+}
+
+#undef LOCTEXT_NAMESPACE // "FAnimNextGraphItemDetails"
\ No newline at end of file
diff --git a/Engine/Plugins/Experimental/AnimNext/Source/AnimNextEditor/Internal/Common/AnimNextAssetItemDetails.h b/Engine/Plugins/Experimental/AnimNext/Source/AnimNextEditor/Internal/Common/AnimNextAssetItemDetails.h
new file mode 100644
--- /dev/null
+++ b/Engine/Plugins/Experimental/AnimNext/Source/AnimNextEditor/Internal/Common/AnimNextAssetItemDetails.h
@@ -0,0 +1,24 @@
+// Copyright Epic Games, Inc. All Rights Reserved.
+
+#pragma once
+
+#include "IWorkspaceOutlinerItemDetails.h"
+
+class UAnimNextRigVMAssetEditorData;
+enum class EAnimNextEditorDataNotifType : uint8;
+
+namespace UE::AnimNext::Editor
+{
+
+class FAnimNextAssetItemDetails : public UE::Workspace::IWorkspaceOutlinerItemDetails
+{
+public:
+	FAnimNextAssetItemDetails() = default;
+
+	ANIMNEXTEDITOR_API virtual const FSlateBrush* GetItemIcon(const FWorkspaceOutlinerItemExport& Export) const override;
+
+	static void RegisterToolMenuExtensions();
+	static void UnregisterToolMenuExtensions();
+};
+
+}
diff --git a/Engine/Plugins/Experimental/AnimNext/Source/AnimNextEditor/Private/AnimNextEditorModule.cpp b/Engine/Plugins/Experimental/AnimNext/Source/AnimNextEditor/Private/AnimNextEditorModule.cpp
--- a/Engine/Plugins/Experimental/AnimNext/Source/AnimNextEditor/Private/AnimNextEditorModule.cpp
+++ b/Engine/Plugins/Experimental/AnimNext/Source/AnimNextEditor/Private/AnimNextEditorModule.cpp
@@ -12,12 +12,10 @@
 #include "Editor/RigVMEditorTools.h"
 #include "Framework/Application/SlateApplication.h"
 #include "Module/AnimNextModule.h"
-#include "Graph/AnimNextGraphPanelNodeFactory.h"
-#include "Graph/AnimNextEdGraphNodeCustomization.h"
+#include "Common/AnimNextEdGraphNodeCustomization.h"
 #include "Module/AnimNextModule_EditorData.h"
 #include "AnimNextEdGraphNode.h"
-#include "Graph/TraitEditorTabSummoner.h"
-#include "Graph/AnimNextCompilerResultsTabSummoner.h"
+#include "Common/AnimNextCompilerResultsTabSummoner.h"
 #include "HAL/PlatformApplicationMisc.h"
 #include "Variables/VariableCustomization.h"
 #include "Variables/VariableBindingPropertyCustomization.h"
@@ -36,20 +34,17 @@
 #include "Param/ObjectPropertyLocatorEditor.h"
 #include "Variables/SAddVariablesDialog.h"
 #include "AnimNextAssetWorkspaceAssetUserData.h"
-#include "Graph/AnimNextGraphItemDetails.h"
-#include "Graph/AnimNextCollapseNodeItemDetails.h"
-#include "Graph/AnimNextFunctionItemDetails.h"
+#include "Common/AnimNextGraphItemDetails.h"
+#include "Common/AnimNextCollapseNodeItemDetails.h"
+#include "Common/AnimNextFunctionItemDetails.h"
 #include "Param/AnimNextActorLocatorEditor.h"
 #include "IWorkspaceEditor.h"
-#include "Entries/AnimNextAnimationGraphEntry.h"
 #include "Framework/Docking/LayoutExtender.h"
-#include "Graph/AnimNextAnimationGraph.h"
-#include "Graph/AnimNextAnimationGraph_EditorData.h"
 #include "Common/AnimNextAssetItemDetails.h"
 #include "Common/AnimNextRigVMAssetEditorDataCustomization.h"
 #include "DataInterface/AnimNextDataInterface.h"
 #include "Entries/AnimNextEventGraphEntry.h"
-#include "Graph/RigVMTrait_AnimNextPublicVariablesEditor.h"
+#include "Logging/MessageLog.h"
 #include "Module/RigUnit_AnimNextModuleEvents.h"
 #include "Variables/SVariablesView.h"
 #include "Variables/VariableOverrideCommands.h"
@@ -63,7 +58,6 @@ namespace UE::AnimNext::Editor
 
 void FAnimNextEditorModule::StartupModule()
 {
-	FPublicVariablesImpl::Register();
 	FVariableOverrideCommands::Register();
 
 	// Register settings for user editing
@@ -93,17 +87,10 @@ void FAnimNextEditorModule::StartupModule()
 	PropertyModule.RegisterCustomClassLayout("AnimNextRigVMAssetEditorData",
 		FOnGetDetailCustomizationInstance::CreateLambda([] { return MakeShared<FAnimNextRigVMAssetEditorDataCustomization>(); }));
 
-	AnimNextGraphPanelNodeFactory = MakeShared<FAnimNextGraphPanelNodeFactory>();
-	FEdGraphUtilities::RegisterVisualNodeFactory(AnimNextGraphPanelNodeFactory);
-
 	Workspace::IWorkspaceEditorModule& WorkspaceEditorModule = FModuleManager::Get().LoadModuleChecked<Workspace::IWorkspaceEditorModule>("WorkspaceEditor");
 
 	WorkspaceEditorModule.OnRegisterTabsForEditor().AddLambda([](FWorkflowAllowedTabSet& TabFactories, const TSharedRef<FTabManager>& InTabManager, TSharedPtr<UE::Workspace::IWorkspaceEditor> InEditorPtr)
 		{
-			TSharedRef<FTraitEditorTabSummoner> TraitEditorTabSummoner = MakeShared<FTraitEditorTabSummoner>(InEditorPtr);
-			TabFactories.RegisterFactory(TraitEditorTabSummoner);
-			TraitEditorTabSummoner->RegisterTabSpawner(InTabManager, nullptr);
-
 			TSharedRef<FAnimNextCompilerResultsTabSummoner> CompilerResultsTabSummoner = MakeShared<FAnimNextCompilerResultsTabSummoner>(InEditorPtr);
 			TabFactories.RegisterFactory(CompilerResultsTabSummoner);
 			CompilerResultsTabSummoner->RegisterTabSpawner(InTabManager, nullptr);
@@ -115,9 +102,6 @@ void FAnimNextEditorModule::StartupModule()
 
 	WorkspaceEditorModule.OnExtendTabs().AddLambda([](FLayoutExtender& InLayoutExtender, TSharedPtr<UE::Workspace::IWorkspaceEditor> InEditorPtr)
 	{
-		FTabManager::FTab TraitEditorTab(FTabId(TraitEditorTabName), ETabState::ClosedTab);
-		InLayoutExtender.ExtendLayout(FTabId(Workspace::WorkspaceTabs::TopRightDocumentArea), ELayoutExtensionPosition::After, TraitEditorTab);
-
 		FTabManager::FTab CompilerResultsTab(FTabId(CompilerResultsTabName), ETabState::ClosedTab);
 		InLayoutExtender.ExtendLayout(FTabId(Workspace::WorkspaceTabs::BottomMiddleDocumentArea), ELayoutExtensionPosition::After, CompilerResultsTab);
 
@@ -182,14 +166,12 @@ void FAnimNextEditorModule::StartupModule()
 
 	const TSharedPtr<FAnimNextAssetItemDetails> AssetItemDetails = MakeShared<FAnimNextAssetItemDetails>();
 	WorkspaceModule.RegisterWorkspaceItemDetails(Workspace::FOutlinerItemDetailsId(FAnimNextModuleOutlinerData::StaticStruct()->GetFName()), AssetItemDetails);
-	WorkspaceModule.RegisterWorkspaceItemDetails(Workspace::FOutlinerItemDetailsId(FAnimNextAnimationGraphOutlinerData::StaticStruct()->GetFName()), AssetItemDetails);
 
 	FAnimNextGraphItemDetails::RegisterToolMenuExtensions();
 	FAnimNextAssetItemDetails::RegisterToolMenuExtensions();
 
 	SupportedAssetClasses.Append(
 		{
-			UAnimNextAnimationGraph::StaticClass()->GetClassPathName(),
 			UAnimNextModule::StaticClass()->GetClassPathName(),
 			UAnimNextDataInterface::StaticClass()->GetClassPathName()
 		});
@@ -207,8 +189,6 @@ void FAnimNextEditorModule::ShutdownModule()
 		PropertyModule.UnregisterCustomClassLayout("AnimNextRigVMAssetEditorData");
 	}
 
-	FEdGraphUtilities::UnregisterVisualNodeFactory(AnimNextGraphPanelNodeFactory);
-
 	UnregisterWorkspaceDocumentTypes();
 
 	if(FModuleManager::Get().IsModuleLoaded("UniversalObjectLocatorEditor"))
@@ -231,7 +211,6 @@ void FAnimNextEditorModule::ShutdownModule()
 		WorkspaceModule.UnregisterWorkspaceItemDetails(Workspace::FOutlinerItemDetailsId(FAnimNextGraphFunctionOutlinerData::StaticStruct()->GetFName()));
 		FAnimNextGraphItemDetails::UnregisterToolMenuExtensions();
 		WorkspaceModule.UnregisterWorkspaceItemDetails(Workspace::FOutlinerItemDetailsId(FAnimNextModuleOutlinerData::StaticStruct()->GetFName()));
-		WorkspaceModule.UnregisterWorkspaceItemDetails(Workspace::FOutlinerItemDetailsId(FAnimNextAnimationGraphOutlinerData::StaticStruct()->GetFName()));
 		FAnimNextAssetItemDetails::UnregisterToolMenuExtensions();
 	}
 
@@ -291,21 +270,7 @@ void FAnimNextEditorModule::RegisterWorkspaceDocumentTypes(Workspace::IWorkspace
 			}));
 
 	WorkspaceEditorModule.RegisterObjectDocumentType(FTopLevelAssetPath(TEXT("/Script/AnimNext.AnimNextModule")), AnimNextModuleDocumentArgs);
-
-	// --- AnimNextAnimationGraph ---
-	Workspace::FObjectDocumentArgs AnimNextAnimationGraphDocumentArgs(
-			Workspace::FOnRedirectWorkspaceContext::CreateLambda([](UObject* InObject)
-			{
-				UAnimNextAnimationGraph* AnimationGraph = CastChecked<UAnimNextAnimationGraph>(InObject);
-				UAnimNextAnimationGraph_EditorData* EditorData = UncookedOnly::FUtils::GetEditorData<UAnimNextAnimationGraph_EditorData>(AnimationGraph);
-
-				// Redirect to the inner graph
-				UAnimNextAnimationGraphEntry* AnimationGraphEntry = CastChecked<UAnimNextAnimationGraphEntry>(EditorData->FindEntry(FRigUnit_AnimNextGraphRoot::DefaultEntryPoint));
-				return AnimationGraphEntry->GetEdGraph();
-			}));
-
-	WorkspaceEditorModule.RegisterObjectDocumentType(FTopLevelAssetPath(TEXT("/Script/AnimNext.AnimNextAnimationGraph")), AnimNextAnimationGraphDocumentArgs);
-
+	
 	// --- AnimNextEdGraph ---
 	Workspace::FGraphDocumentWidgetArgs GraphArgs;
 	GraphArgs.SpawnLocation = Workspace::WorkspaceTabs::TopMiddleDocumentArea;
@@ -831,7 +796,6 @@ void FAnimNextEditorModule::UnregisterWorkspaceDocumentTypes()
 	{
 		Workspace::IWorkspaceEditorModule& WorkspaceEditorModule = FModuleManager::LoadModuleChecked<Workspace::IWorkspaceEditorModule>("WorkspaceEditor");
 		WorkspaceEditorModule.UnregisterObjectDocumentType(FTopLevelAssetPath(TEXT("/Script/AnimNext.AnimNextModule")));
-		WorkspaceEditorModule.UnregisterObjectDocumentType(FTopLevelAssetPath(TEXT("/Script/AnimNext.AnimNextAnimationGraph")));
 		WorkspaceEditorModule.UnregisterObjectDocumentType(FTopLevelAssetPath(TEXT("/Script/AnimNextUncookedOnly.AnimNextEdGraph")));
 	}
 }
diff --git a/Engine/Plugins/Experimental/AnimNext/Source/AnimNextEditor/Private/AnimNextEditorModule.h b/Engine/Plugins/Experimental/AnimNext/Source/AnimNextEditor/Private/AnimNextEditorModule.h
--- a/Engine/Plugins/Experimental/AnimNext/Source/AnimNextEditor/Private/AnimNextEditorModule.h
+++ b/Engine/Plugins/Experimental/AnimNext/Source/AnimNextEditor/Private/AnimNextEditorModule.h
@@ -5,13 +5,12 @@
 #include "IAnimNextEditorModule.h"
 #include "UObject/TopLevelAssetPath.h"
 
-class FAnimNextGraphPanelNodeFactory;
 class UAnimNextWorkspaceSchema;
+class FAnimNextGraphPanelNodeFactory;
 
 namespace UE::AnimNext::Editor
 {
 	class FParamNamePropertyTypeIdentifier;
-	class FParametersGraphPanelPinFactory;
 	class FLocatorContext;
 }
 
@@ -40,9 +39,6 @@ private:
 	void UnregisterWorkspaceDocumentTypes();
 
 private:
-	/** Node factory for the AnimNext graph */
-	TSharedPtr<FAnimNextGraphPanelNodeFactory> AnimNextGraphPanelNodeFactory;
-
 	/** Type identifier for parameter names */
 	TSharedPtr<FParamNamePropertyTypeIdentifier> Identifier;
 
diff --git a/Engine/Plugins/Experimental/AnimNext/Source/AnimNextEditor/Private/Common/AnimNextEntryAssetDefinitions.cpp b/Engine/Plugins/Experimental/AnimNext/Source/AnimNextEditor/Private/Common/AnimNextEntryAssetDefinitions.cpp
--- a/Engine/Plugins/Experimental/AnimNext/Source/AnimNextEditor/Private/Common/AnimNextEntryAssetDefinitions.cpp
+++ b/Engine/Plugins/Experimental/AnimNext/Source/AnimNextEditor/Private/Common/AnimNextEntryAssetDefinitions.cpp
@@ -10,12 +10,6 @@ FText UAssetDefinition_AnimNextVariableEntry::GetObjectDisplayNameText(UObject*
 	return Parameter->GetDisplayName();
 }
 
-FText UAssetDefinition_AnimNextAnimationGraphEntry::GetObjectDisplayNameText(UObject* Object) const
-{
-	UAnimNextRigVMAssetEntry* Parameter = CastChecked<UAnimNextRigVMAssetEntry>(Object);
-	return Parameter->GetDisplayName();
-}
-
 FText UAssetDefinition_AnimNextEventGraphEntry::GetObjectDisplayNameText(UObject* Object) const
 {
 	UAnimNextRigVMAssetEntry* Variable = CastChecked<UAnimNextRigVMAssetEntry>(Object);
diff --git a/Engine/Plugins/Experimental/AnimNext/Source/AnimNextEditor/Private/Common/AnimNextEntryAssetDefinitions.h b/Engine/Plugins/Experimental/AnimNext/Source/AnimNextEditor/Private/Common/AnimNextEntryAssetDefinitions.h
--- a/Engine/Plugins/Experimental/AnimNext/Source/AnimNextEditor/Private/Common/AnimNextEntryAssetDefinitions.h
+++ b/Engine/Plugins/Experimental/AnimNext/Source/AnimNextEditor/Private/Common/AnimNextEntryAssetDefinitions.h
@@ -5,7 +5,6 @@
 #include "AssetDefinitionDefault.h"
 #include "Entries/AnimNextEventGraphEntry.h"
 #include "Entries/AnimNextVariableEntry.h"
-#include "Entries/AnimNextAnimationGraphEntry.h"
 #include "Entries/AnimNextDataInterfaceEntry.h"
 #include "AnimNextEntryAssetDefinitions.generated.h"
 
@@ -24,19 +23,6 @@ public:
 	virtual FText GetObjectDisplayNameText(UObject* Object) const override;
 };
 
-UCLASS()
-class UAssetDefinition_AnimNextAnimationGraphEntry : public UAssetDefinitionDefault
-{
-	GENERATED_BODY()
-
-public:
-	// UAssetDefinition interface
-	virtual FText GetAssetDisplayName() const override { return LOCTEXT("AnimNextAnimationGraphEntry", "Animation Graph"); }
-	virtual FLinearColor GetAssetColor() const override { return FLinearColor(FColor(128,64,64)); }
-	virtual TSoftClassPtr<UObject> GetAssetClass() const override { return UAnimNextAnimationGraphEntry::StaticClass(); }
-	virtual FText GetObjectDisplayNameText(UObject* Object) const override;
-};
-
 UCLASS()
 class UAssetDefinition_AnimNextEventGraphEntry : public UAssetDefinitionDefault
 {
diff --git a/Engine/Plugins/Experimental/AnimNext/Source/AnimNextEditor/Private/Param/ObjectCastLocatorEditor.cpp b/Engine/Plugins/Experimental/AnimNext/Source/AnimNextEditor/Private/Param/ObjectCastLocatorEditor.cpp
--- a/Engine/Plugins/Experimental/AnimNext/Source/AnimNextEditor/Private/Param/ObjectCastLocatorEditor.cpp
+++ b/Engine/Plugins/Experimental/AnimNext/Source/AnimNextEditor/Private/Param/ObjectCastLocatorEditor.cpp
@@ -3,7 +3,6 @@
 #include "ObjectCastLocatorEditor.h"
 
 #include "ClassViewerModule.h"
-#include "IAnimNextModuleInterface.h"
 #include "UniversalObjectLocatorFragmentTypeHandle.h"
 #include "UniversalObjectLocator.h"
 #include "IUniversalObjectLocatorCustomization.h"
diff --git a/Engine/Plugins/Experimental/AnimNext/Source/AnimNextEditor/Private/Param/ObjectFunctionLocatorEditor.cpp b/Engine/Plugins/Experimental/AnimNext/Source/AnimNextEditor/Private/Param/ObjectFunctionLocatorEditor.cpp
--- a/Engine/Plugins/Experimental/AnimNext/Source/AnimNextEditor/Private/Param/ObjectFunctionLocatorEditor.cpp
+++ b/Engine/Plugins/Experimental/AnimNext/Source/AnimNextEditor/Private/Param/ObjectFunctionLocatorEditor.cpp
@@ -3,7 +3,6 @@
 #include "ObjectFunctionLocatorEditor.h"
 
 #include "ClassViewerModule.h"
-#include "IAnimNextModuleInterface.h"
 #include "UniversalObjectLocatorFragmentTypeHandle.h"
 #include "UniversalObjectLocator.h"
 #include "IUniversalObjectLocatorCustomization.h"
@@ -14,6 +13,7 @@
 #include "Widgets/Input/SComboButton.h"
 #include "Modules/ModuleManager.h"
 #include "Framework/PropertyViewer/IFieldIterator.h"
+#include "Kismet/BlueprintFunctionLibrary.h"
 #include "Widgets/PropertyViewer/SPropertyViewer.h"
 #include "Param/ParamUtils.h"
 #include "Widgets/Layout/SBox.h"
diff --git a/Engine/Plugins/Experimental/AnimNext/Source/AnimNextEditor/Private/Param/ObjectPropertyLocatorEditor.cpp b/Engine/Plugins/Experimental/AnimNext/Source/AnimNextEditor/Private/Param/ObjectPropertyLocatorEditor.cpp
--- a/Engine/Plugins/Experimental/AnimNext/Source/AnimNextEditor/Private/Param/ObjectPropertyLocatorEditor.cpp
+++ b/Engine/Plugins/Experimental/AnimNext/Source/AnimNextEditor/Private/Param/ObjectPropertyLocatorEditor.cpp
@@ -3,7 +3,6 @@
 #include "ObjectPropertyLocatorEditor.h"
 
 #include "ClassViewerModule.h"
-#include "IAnimNextModuleInterface.h"
 #include "IUniversalObjectLocatorCustomization.h"
 #include "UniversalObjectLocator.h"
 #include "UniversalObjectLocatorFragmentTypeHandle.h"
diff --git a/Engine/Plugins/Experimental/AnimNext/Source/AnimNextEditor/Private/Workspace/AnimNextWorkspaceSchema.cpp b/Engine/Plugins/Experimental/AnimNext/Source/AnimNextEditor/Private/Workspace/AnimNextWorkspaceSchema.cpp
--- a/Engine/Plugins/Experimental/AnimNext/Source/AnimNextEditor/Private/Workspace/AnimNextWorkspaceSchema.cpp
+++ b/Engine/Plugins/Experimental/AnimNext/Source/AnimNextEditor/Private/Workspace/AnimNextWorkspaceSchema.cpp
@@ -3,8 +3,6 @@
 #include "AnimNextWorkspaceSchema.h"
 
 #include "AnimNextEditorModule.h"
-#include "Graph/AnimNextAnimationGraph.h"
-#include "Module/AnimNextModule.h"
 
 #define LOCTEXT_NAMESPACE "AnimNextWorkspaceSchema"
 
diff --git a/Engine/Plugins/Experimental/AnimNext/Source/AnimNextTestSuite/Private/AnimNextDataInterfaceTest.cpp b/Engine/Plugins/Experimental/AnimNext/Source/AnimNextTestSuite/Private/AnimNextDataInterfaceTest.cpp
--- a/Engine/Plugins/Experimental/AnimNext/Source/AnimNextTestSuite/Private/AnimNextDataInterfaceTest.cpp
+++ b/Engine/Plugins/Experimental/AnimNext/Source/AnimNextTestSuite/Private/AnimNextDataInterfaceTest.cpp
@@ -10,8 +10,6 @@
 #include "DataInterface/AnimNextDataInterface_EditorData.h"
 #include "Entries/AnimNextDataInterfaceEntry.h"
 #include "Entries/AnimNextVariableEntry.h"
-#include "Graph/AnimNextAnimationGraph.h"
-#include "Graph/AnimNextAnimationGraphFactory.h"
 
 // AnimNext Data Interface Tests
 
diff --git a/Engine/Plugins/Experimental/AnimNext/Source/AnimNextTestSuite/Private/AnimNextEditorTest.cpp b/Engine/Plugins/Experimental/AnimNext/Source/AnimNextTestSuite/Private/AnimNextEditorTest.cpp
--- a/Engine/Plugins/Experimental/AnimNext/Source/AnimNextTestSuite/Private/AnimNextEditorTest.cpp
+++ b/Engine/Plugins/Experimental/AnimNext/Source/AnimNextTestSuite/Private/AnimNextEditorTest.cpp
@@ -1,16 +1,13 @@
 // Copyright Epic Games, Inc. All Rights Reserved.
 
 #include "CoreMinimal.h"
+#include "AnimNextRigVMAssetEditorData.h"
 #include "AnimNextTest.h"
 #include "UncookedOnlyUtils.h"
 #include "Misc/AutomationTest.h"
 #include "Animation/AnimSequence.h"
-#include "Graph/AnimNextAnimationGraph.h"
 #include "Entries/AnimNextVariableEntry.h"
-#include "Graph/AnimNextAnimationGraph_EditorData.h"
-#include "Entries/AnimNextAnimationGraphEntry.h"
 #include "Entries/AnimNextEventGraphEntry.h"
-#include "Graph/AnimNextAnimationGraphFactory.h"
 #include "Module/AnimNextModuleFactory.h"
 #include "Module/RigUnit_AnimNextModuleEvents.h"
 #if WITH_EDITOR
@@ -40,7 +37,6 @@ bool FEditor_Variables::RunTest(const FString& InParameters)
 	
 	FFactoryAndClass FactoryClassPairs[] =
 	{
-		{ UAnimNextAnimationGraphFactory::StaticClass(), UAnimNextAnimationGraph::StaticClass() },
 		{ UAnimNextModuleFactory::StaticClass(), UAnimNextModule::StaticClass() },
 	};
 
@@ -144,7 +140,6 @@ bool FEditor_Graphs::RunTest(const FString& InParameters)
 	
 	FTestSettings TestSettings[] =
 	{
-		{ UAnimNextAnimationGraphFactory::StaticClass(), UAnimNextAnimationGraph::StaticClass(), false, false, true, true },
 		{ UAnimNextModuleFactory::StaticClass(), UAnimNextModule::StaticClass(), true, true, false, false },
 	};
 
@@ -193,61 +188,26 @@ bool FEditor_Graphs::RunTest(const FString& InParameters)
 			AddExpectedError(TEXT("Cannot add an event graph to this asset - entry is not allowed"));
 			UAnimNextEventGraphEntry* EventGraphEntry = EditorData->AddEventGraph(TEXT("PrePhysics"), FRigUnit_AnimNextPrePhysicsEvent::StaticStruct());
 		}
-		// AddAnimationGraph
-		if(TestSetting.bAnimGraphsAllowed)
-		{
-			UAnimNextAnimationGraphEntry* AnimationGraphEntry = nullptr;
-			if(TestSetting.bExpectExistingAnimGraph)
-			{
-				AnimationGraphEntry = Cast<UAnimNextAnimationGraphEntry>(EditorData->FindEntry(FRigUnit_AnimNextGraphRoot::DefaultEntryPoint));
-				UE_RETURN_ON_ERROR(AnimationGraphEntry != nullptr, TEXT("Could not find existing animation graph."));
-			}
-			else
-			{
-				AnimationGraphEntry = EditorData->AddAnimationGraph(FRigUnit_AnimNextGraphRoot::DefaultEntryPoint);
-				UE_RETURN_ON_ERROR(AnimationGraphEntry != nullptr, "FEditor_Graphs -> Could not add animation graph.");
-			}
-
-			URigVMGraph* RigVMGraph = AnimationGraphEntry->GetRigVMGraph();
-			UE_RETURN_ON_ERROR(RigVMGraph->GetNodes().Num() == 1, TEXT("Unexpected number of nodes in new animation graph."));
-
-			{
-				FScopedTransaction Transaction(FText::GetEmpty());
-				bool bRemovedEventGraph = EditorData->RemoveEntry(AnimationGraphEntry);
-				UE_RETURN_ON_ERROR(bRemovedEventGraph, "FEditor_Graphs -> Could not remove animation graph.");
-			}
-
-			GEditor->UndoTransaction();
-
-			UAnimNextAnimationGraphEntry* FoundEventGraphEntry = Cast<UAnimNextAnimationGraphEntry>(EditorData->FindEntry(FRigUnit_AnimNextGraphRoot::DefaultEntryPoint));
-			UE_RETURN_ON_ERROR(FoundEventGraphEntry != nullptr, "FEditor_Graphs -> Could not find animation graph post-undo.");
-		}
-		else
-		{
-			AddExpectedError(TEXT("Cannot add an animation graph to this asset - entry is not allowed"));
-			UAnimNextAnimationGraphEntry* AnimationGraphEntry = EditorData->AddAnimationGraph(TEXT("TestAnimationGraph"));
-		}
 	}
 
 	return true;
 }
 
 
-IMPLEMENT_SIMPLE_AUTOMATION_TEST(FEditor_Variables_Python, "Animation.AnimNext.Editor.PythonParameters", EAutomationTestFlags::EditorContext | EAutomationTestFlags::EngineFilter)
+IMPLEMENT_SIMPLE_AUTOMATION_TEST(FEditor_Variables_Python, "Animation.AnimNext.Editor.Python.Variables", EAutomationTestFlags::EditorContext | EAutomationTestFlags::EngineFilter)
 
 bool FEditor_Variables_Python::RunTest(const FString& InParameters)
 {
 	using namespace UE::AnimNext;
 
-	// FIXME: disabled for now as we seemingly cant call hoisted base-class functions on derived types from Python?
-/*	const TCHAR* Script = TEXT(
+	const TCHAR* Script = TEXT(
 		"asset_tools = unreal.AssetToolsHelpers.get_asset_tools()\n"
-		"animation_graph = unreal.AssetTools.create_asset(asset_tools, asset_name = \"TestAnimationGraph\", package_path = \"/Game/\", asset_class = unreal.AnimNextAnimationGraph, factory = unreal.AnimNextAnimationGraphFactory())\n"
-		"animation_graph.add_variable(name = \"TestParam\", value_type = unreal.PropertyBagPropertyType.BOOL, container_type = unreal.PropertyBagContainerType.NONE)\n"
-		"unreal.EditorAssetLibrary.delete_loaded_asset(animation_graph)\n"
+		"animation_module = unreal.AssetTools.create_asset(asset_tools, asset_name = \"TestModule\", package_path = \"/Game/\", asset_class = unreal.AnimNextModule, factory = unreal.AnimNextModuleFactory())\n"
+		"animation_module.add_variable(name = \"TestParam\", value_type = unreal.PropertyBagPropertyType.BOOL, container_type = unreal.PropertyBagContainerType.NONE)\n"
+		"unreal.EditorAssetLibrary.delete_loaded_asset(animation_module)\n"
 	);
 
-	IPythonScriptPlugin::Get()->ExecPythonCommand(Script);*/
+	IPythonScriptPlugin::Get()->ExecPythonCommand(Script);
 
 	FUtils::CleanupAfterTests();
 
diff --git a/Engine/Plugins/Experimental/AnimNext/Source/AnimNextUncookedOnly/AnimNextUncookedOnly.Build.cs b/Engine/Plugins/Experimental/AnimNext/Source/AnimNextUncookedOnly/AnimNextUncookedOnly.Build.cs
--- a/Engine/Plugins/Experimental/AnimNext/Source/AnimNextUncookedOnly/AnimNextUncookedOnly.Build.cs
+++ b/Engine/Plugins/Experimental/AnimNext/Source/AnimNextUncookedOnly/AnimNextUncookedOnly.Build.cs
@@ -49,6 +49,13 @@ namespace UnrealBuildTool.Rules
 					}
 				);
 
+				PublicDependencyModuleNames.AddRange(
+					new string[]
+					{
+						"WorkspaceEditor",
+					}
+				);
+				
 				PrivateIncludePathModuleNames.AddRange(
 					new string[]
 					{
diff --git a/Engine/Plugins/Experimental/AnimNext/Source/AnimNextUncookedOnly/Internal/AnimNextAssetWorkspaceAssetUserData.h b/Engine/Plugins/Experimental/AnimNext/Source/AnimNextUncookedOnly/Internal/AnimNextAssetWorkspaceAssetUserData.h
--- a/Engine/Plugins/Experimental/AnimNext/Source/AnimNextUncookedOnly/Internal/AnimNextAssetWorkspaceAssetUserData.h
+++ b/Engine/Plugins/Experimental/AnimNext/Source/AnimNextUncookedOnly/Internal/AnimNextAssetWorkspaceAssetUserData.h
@@ -7,7 +7,6 @@
 #include "Param/ParamType.h"
 #include "IAnimNextRigVMGraphInterface.h"
 #include "RigVMModel/Nodes/RigVMFunctionReferenceNode.h"
-#include "Graph/AnimNextAnimationGraph.h"
 #include "Module/AnimNextModule.h"
 #include "DataInterface/AnimNextDataInterface.h"
 
@@ -41,19 +40,6 @@ struct FAnimNextModuleOutlinerData : public FAnimNextRigVMAssetOutlinerData
 	}
 };
 
-USTRUCT()
-struct FAnimNextAnimationGraphOutlinerData : public FAnimNextRigVMAssetOutlinerData
-{
-	GENERATED_BODY()
-
-	FAnimNextAnimationGraphOutlinerData() = default;
-
-	UAnimNextAnimationGraph* GetAnimationGraph() const
-	{
-		return Cast<UAnimNextAnimationGraph>(Asset);
-	}
-};
-
 USTRUCT()
 struct FAnimNextDataInterfaceOutlinerData : public FAnimNextRigVMAssetOutlinerData
 {
@@ -131,7 +117,7 @@ struct FAnimNextGraphOutlinerData : public FAnimNextAssetEntryOutlinerData
 	TScriptInterface<IAnimNextRigVMGraphInterface> GraphInterface;
 };
 
-UCLASS()
+UCLASS(MinimalAPI)
 class UAnimNextAssetWorkspaceAssetUserData : public UAssetUserData
 {
 public:
diff --git a/Engine/Plugins/Experimental/AnimNext/Source/AnimNextUncookedOnly/Internal/AnimNextEdGraphNode.h b/Engine/Plugins/Experimental/AnimNext/Source/AnimNextUncookedOnly/Internal/AnimNextEdGraphNode.h
--- a/Engine/Plugins/Experimental/AnimNext/Source/AnimNextUncookedOnly/Internal/AnimNextEdGraphNode.h
+++ b/Engine/Plugins/Experimental/AnimNext/Source/AnimNextUncookedOnly/Internal/AnimNextEdGraphNode.h
@@ -7,35 +7,10 @@
 
 #include "AnimNextEdGraphNode.generated.h"
 
-// EdGraphNode representation for AnimNext nodes
-// A node can hold a trait stack or a trait entry
+// EdGraphNode representation for AnimNext nodes.
+// No longer has a specific use, but keeping this class around just in case.
 UCLASS(MinimalAPI)
 class UAnimNextEdGraphNode : public URigVMEdGraphNode
 {
 	GENERATED_BODY()
-
-public:
-	//////////////////////////////////////////////////////////////////////////
-	// UEdGraphNode implementation
-	virtual void GetNodeContextMenuActions(UToolMenu* Menu, UGraphNodeContextMenuContext* Context) const override;
-
-	//////////////////////////////////////////////////////////////////////////
-	// URigVMEdGraphNode implementation
-	virtual void ConfigurePin(UEdGraphPin* EdGraphPin, const URigVMPin* ModelPin) const override;
-
-	//////////////////////////////////////////////////////////////////////////
-	// Our implementation
-
-	// Returns whether this node is a trait stack or not
-	ANIMNEXTUNCOOKEDONLY_API bool IsTraitStack() const;
-
-private:
-	// Populates the SubMenu with entries for each trait that can be added through the context menu
-	void BuildAddTraitContextMenu(UToolMenu* SubMenu);
-
-	// Returns whether this node is an animation graph host node
-	bool IsRunGraphNode() const;
-
-	// Builds a menu to expose the public variables of a specific animation graph asset
-	void BuildExposeVariablesContextMenu(UToolMenu* SubMenu);
 };
diff --git a/Engine/Plugins/Experimental/AnimNext/Source/AnimNextUncookedOnly/Internal/AnimNextRigVMAssetEditorData.h b/Engine/Plugins/Experimental/AnimNext/Source/AnimNextUncookedOnly/Internal/AnimNextRigVMAssetEditorData.h
--- a/Engine/Plugins/Experimental/AnimNext/Source/AnimNextUncookedOnly/Internal/AnimNextRigVMAssetEditorData.h
+++ b/Engine/Plugins/Experimental/AnimNext/Source/AnimNextUncookedOnly/Internal/AnimNextRigVMAssetEditorData.h
@@ -48,7 +48,9 @@ namespace UE::AnimNext::Editor
 namespace UE::AnimNext::Tests
 {
 	class FEditor_Graphs;
+	class FEditor_AnimGraph_Graphs;
 	class FEditor_Variables;
+	class FEditor_AnimGraph_Variables;
 	class FVariables;
 	class FDataInterfaceCompile;
 }
@@ -98,10 +100,6 @@ public:
 	UFUNCTION(BlueprintCallable, Category = "AnimNext|Entries", meta=(ScriptMethod))
 	static ANIMNEXTUNCOOKEDONLY_API bool RemoveAllEntries(UAnimNextRigVMAsset* InAsset, bool bSetupUndoRedo = true, bool bPrintPythonCommand = true);
 
-	/** Adds an animation graph to an AnimNext asset */
-	UFUNCTION(BlueprintCallable, Category = "AnimNext|Entries", meta=(ScriptMethod))
-	static ANIMNEXTUNCOOKEDONLY_API UAnimNextAnimationGraphEntry* AddAnimationGraph(UAnimNextRigVMAsset* InAsset, FName InName, bool bSetupUndoRedo = true, bool bPrintPythonCommand = true);
-
 	/** Adds a parameter to an AnimNext asset */
 	UFUNCTION(BlueprintCallable, Category = "AnimNext|Entries", meta=(ScriptMethod))
 	static ANIMNEXTUNCOOKEDONLY_API UAnimNextVariableEntry* AddVariable(UAnimNextRigVMAsset* InAsset, FName InName, EPropertyBagPropertyType InValueType, EPropertyBagContainerType InContainerType = EPropertyBagContainerType::None, const UObject* InValueTypeObject = nullptr, const FString& InDefaultValue = TEXT(""), bool bSetupUndoRedo = true, bool bPrintPythonCommand = true);
@@ -122,9 +120,6 @@ class ANIMNEXTUNCOOKEDONLY_API UAnimNextRigVMAssetEditorData : public UObject, p
 	GENERATED_BODY()
 
 public:
-	/** Adds an animation graph to this asset */
-	UAnimNextAnimationGraphEntry* AddAnimationGraph(FName InName, bool bSetupUndoRedo = true, bool bPrintPythonCommand = true);
-
 	/** Adds a parameter to this asset */
 	UAnimNextVariableEntry* AddVariable(FName InName, FAnimNextParamType InType, const FString& InDefaultValue = TEXT(""), bool bSetupUndoRedo = true, bool bPrintPythonCommand = true);
 
@@ -134,6 +129,9 @@ public:
 	/** Adds a data interface to this asset */
 	UAnimNextDataInterfaceEntry* AddDataInterface(UAnimNextDataInterface* InDataInterface, bool bSetupUndoRedo = true, bool bPrintPythonCommand = true);
 
+	// Find an entry by name
+	UAnimNextRigVMAssetEntry* FindEntry(FName InName) const;
+
 	// Get the external packaging status of this asset
 	bool IsUsingExternalPackages() const { return bUsesExternalPackages; }
 
@@ -160,6 +158,8 @@ protected:
 	friend class UAnimNextEdGraph;
 	friend class UE::AnimNext::Tests::FEditor_Graphs;
 	friend class UE::AnimNext::Tests::FEditor_Variables;
+	friend class UE::AnimNext::Tests::FEditor_AnimGraph_Graphs;
+	friend class UE::AnimNext::Tests::FEditor_AnimGraph_Variables;
 	friend class UE::AnimNext::Tests::FVariables;
 	friend class UE::AnimNext::Tests::FDataInterfaceCompile;
 	friend class UE::AnimNext::Editor::FVariableCustomization;
@@ -278,6 +278,10 @@ protected:
 	// Class to use when instantiating AssetUserData for the EditorData instance
 	virtual TSubclassOf<UAssetUserData> GetAssetUserDataClass() const;
 
+	// Override point called during initialization (PostLoad/PostDuplicate) used for setting up asset user data.
+	// By default this instantiates any asset user data that is missing according to GetAssetUserDataClass().
+	virtual void InitializeAssetUserData();
+
 	// Get all the kinds of entry for this asset
 	virtual TConstArrayView<TSubclassOf<UAnimNextRigVMAssetEntry>> GetEntryClasses() const PURE_VIRTUAL(UAnimNextRigVMAssetEditorData::GetEntryClasses, return {};)
 
@@ -381,9 +385,6 @@ protected:
 			}
 		}
 	}
-	
-	// Find an entry by name
-	UAnimNextRigVMAssetEntry* FindEntry(FName InName) const;
 
 	// Remove an entry from the asset
 	// @return true if the item was removed
diff --git a/Engine/Plugins/Experimental/AnimNext/Source/AnimNextUncookedOnly/Internal/AnimNextRigVMAssetSchema.h b/Engine/Plugins/Experimental/AnimNext/Source/AnimNextUncookedOnly/Internal/AnimNextRigVMAssetSchema.h
--- a/Engine/Plugins/Experimental/AnimNext/Source/AnimNextUncookedOnly/Internal/AnimNextRigVMAssetSchema.h
+++ b/Engine/Plugins/Experimental/AnimNext/Source/AnimNextUncookedOnly/Internal/AnimNextRigVMAssetSchema.h
@@ -7,7 +7,7 @@
 #include "AnimNextRigVMAssetSchema.generated.h"
 
 UCLASS()
-class UAnimNextRigVMAssetSchema : public URigVMSchema
+class ANIMNEXTUNCOOKEDONLY_API UAnimNextRigVMAssetSchema : public URigVMSchema
 {
 	GENERATED_BODY()
 
diff --git a/Engine/Plugins/Experimental/AnimNext/Source/AnimNextUncookedOnly/Internal/DataInterface/AnimNextDataInterface_EditorData.h b/Engine/Plugins/Experimental/AnimNext/Source/AnimNextUncookedOnly/Internal/DataInterface/AnimNextDataInterface_EditorData.h
--- a/Engine/Plugins/Experimental/AnimNext/Source/AnimNextUncookedOnly/Internal/DataInterface/AnimNextDataInterface_EditorData.h
+++ b/Engine/Plugins/Experimental/AnimNext/Source/AnimNextUncookedOnly/Internal/DataInterface/AnimNextDataInterface_EditorData.h
@@ -3,7 +3,6 @@
 #pragma once
 
 #include "CoreMinimal.h"
-#include "AnimNextController.h"
 #include "AnimNextExecuteContext.h"
 #include "AnimNextRigVMAssetEditorData.h"
 #include "AnimNextDataInterface_EditorData.generated.h"
@@ -18,7 +17,6 @@ class ANIMNEXTUNCOOKEDONLY_API UAnimNextDataInterface_EditorData : public UAnimN
 
 protected:
 	// UAnimNextRigVMAssetEditorData interface
-	virtual TSubclassOf<URigVMController> GetControllerClass() const override { return UAnimNextController::StaticClass(); }
 	virtual UScriptStruct* GetExecuteContextStruct() const override { return FAnimNextExecuteContext::StaticStruct(); }
 	virtual TConstArrayView<TSubclassOf<UAnimNextRigVMAssetEntry>> GetEntryClasses() const override;
 	virtual void CustomizeNewAssetEntry(UAnimNextRigVMAssetEntry* InNewEntry) const override;
diff --git a/Engine/Plugins/Experimental/AnimNext/Source/AnimNextUncookedOnly/Internal/Entries/AnimNextEventGraphEntry.h b/Engine/Plugins/Experimental/AnimNext/Source/AnimNextUncookedOnly/Internal/Entries/AnimNextEventGraphEntry.h
--- a/Engine/Plugins/Experimental/AnimNext/Source/AnimNextUncookedOnly/Internal/Entries/AnimNextEventGraphEntry.h
+++ b/Engine/Plugins/Experimental/AnimNext/Source/AnimNextUncookedOnly/Internal/Entries/AnimNextEventGraphEntry.h
@@ -17,6 +17,8 @@ namespace UE::AnimNext::Tests
 {
 	class FEditor_Variables;
 	class FEditor_Graphs;
+	class FEditor_AnimGraph_Variables;
+	class FEditor_AnimGraph_Graphs;
 	class FVariables;
 }
 
@@ -33,7 +35,9 @@ class UAnimNextEventGraphEntry : public UAnimNextRigVMAssetEntry, public IAnimNe
 	friend class UAnimNextRigVMAssetEditorData;
 	friend class UAnimNextModule_EditorData;
 	friend class UE::AnimNext::Tests::FEditor_Variables;
+	friend class UE::AnimNext::Tests::FEditor_AnimGraph_Variables;
 	friend class UE::AnimNext::Tests::FEditor_Graphs;
+	friend class UE::AnimNext::Tests::FEditor_AnimGraph_Graphs;
 	friend class UE::AnimNext::Tests::FVariables;
 	friend class UE::AnimNext::Editor::FAnimNextEditorModule;
 
diff --git a/Engine/Plugins/Experimental/AnimNext/Source/AnimNextUncookedOnly/Internal/Entries/AnimNextRigVMAssetEntry.h b/Engine/Plugins/Experimental/AnimNext/Source/AnimNextUncookedOnly/Internal/Entries/AnimNextRigVMAssetEntry.h
--- a/Engine/Plugins/Experimental/AnimNext/Source/AnimNextUncookedOnly/Internal/Entries/AnimNextRigVMAssetEntry.h
+++ b/Engine/Plugins/Experimental/AnimNext/Source/AnimNextUncookedOnly/Internal/Entries/AnimNextRigVMAssetEntry.h
@@ -20,8 +20,8 @@ namespace UE::AnimNext::Editor
 }
 
 /** Base class that defines an entry in a module, e.g. a parameter or a graph */
-UCLASS(MinimalAPI, BlueprintType, Abstract)
-class UAnimNextRigVMAssetEntry : public UObject
+UCLASS(BlueprintType, Abstract)
+class ANIMNEXTUNCOOKEDONLY_API UAnimNextRigVMAssetEntry : public UObject
 {
 	GENERATED_BODY()
 public:
@@ -57,7 +57,7 @@ public:
 #endif
 
 protected:
-	ANIMNEXTUNCOOKEDONLY_API void BroadcastModified(EAnimNextEditorDataNotifType InType);
+	void BroadcastModified(EAnimNextEditorDataNotifType InType);
 
 	friend class UE::AnimNext::Editor::FVariableCustomization;
 	friend class UE::AnimNext::Editor::SRigVMAssetView;
diff --git a/Engine/Plugins/Experimental/AnimNext/Source/AnimNextUncookedOnly/Internal/Entries/AnimNextVariableEntry.h b/Engine/Plugins/Experimental/AnimNext/Source/AnimNextUncookedOnly/Internal/Entries/AnimNextVariableEntry.h
--- a/Engine/Plugins/Experimental/AnimNext/Source/AnimNextUncookedOnly/Internal/Entries/AnimNextVariableEntry.h
+++ b/Engine/Plugins/Experimental/AnimNext/Source/AnimNextUncookedOnly/Internal/Entries/AnimNextVariableEntry.h
@@ -36,6 +36,7 @@ namespace UE::AnimNext::Editor
 namespace UE::AnimNext::Tests
 {
 	class FEditor_Variables;
+	class FEditor_AnimGraph_Variables;
 	class FDataInterfaceCompile;
 }
 
@@ -54,6 +55,7 @@ class UAnimNextVariableEntry : public UAnimNextRigVMAssetEntry, public IAnimNext
 	friend class UAnimNextDataInterface_EditorData;
 	friend class UAnimNextDataInterfaceEntry;
 	friend class UE::AnimNext::Tests::FEditor_Variables;
+	friend class UE::AnimNext::Tests::FEditor_AnimGraph_Variables;
 	friend class UE::AnimNext::Tests::FDataInterfaceCompile;
 	friend class FAnimationAnimNextRuntimeTest_Variables;
 	friend class UE::AnimNext::Editor::FVariableCustomization;
diff --git a/Engine/Plugins/Experimental/AnimNext/Source/AnimNextUncookedOnly/Internal/Module/AnimNextModule_EditorData.h b/Engine/Plugins/Experimental/AnimNext/Source/AnimNextUncookedOnly/Internal/Module/AnimNextModule_EditorData.h
--- a/Engine/Plugins/Experimental/AnimNext/Source/AnimNextUncookedOnly/Internal/Module/AnimNextModule_EditorData.h
+++ b/Engine/Plugins/Experimental/AnimNext/Source/AnimNextUncookedOnly/Internal/Module/AnimNextModule_EditorData.h
@@ -3,7 +3,6 @@
 #pragma once
 
 #include "CoreMinimal.h"
-#include "AnimNextController.h"
 #include "RigVMModel/RigVMGraph.h"
 #include "AnimNextEdGraph.h"
 #include "DataInterface/AnimNextDataInterface_EditorData.h"
@@ -49,11 +48,7 @@ class UAnimNextModule_EditorData : public UAnimNextDataInterface_EditorData
 	friend class FAnimationAnimNextRuntimeTest_GraphExecuteLatent;
 
 private:
-	// UObject interface
-	virtual void PostLoad() override;
-
 	// UAnimNextRigVMAssetEditorData interface
-	virtual TSubclassOf<URigVMController> GetControllerClass() const override { return UAnimNextController::StaticClass(); }
 	virtual UScriptStruct* GetExecuteContextStruct() const override { return FAnimNextExecuteContext::StaticStruct(); }
 	virtual TConstArrayView<TSubclassOf<UAnimNextRigVMAssetEntry>> GetEntryClasses() const override;
 	virtual void RecompileVM() override;
diff --git a/Engine/Plugins/Experimental/AnimNext/Source/AnimNextUncookedOnly/Internal/UncookedOnlyUtils.h b/Engine/Plugins/Experimental/AnimNext/Source/AnimNextUncookedOnly/Internal/UncookedOnlyUtils.h
--- a/Engine/Plugins/Experimental/AnimNext/Source/AnimNextUncookedOnly/Internal/UncookedOnlyUtils.h
+++ b/Engine/Plugins/Experimental/AnimNext/Source/AnimNextUncookedOnly/Internal/UncookedOnlyUtils.h
@@ -106,7 +106,7 @@ struct FAnimNextAssetRegistryExports
 namespace UE::AnimNext::UncookedOnly
 {
 
-extern TAutoConsoleVariable<bool> CVarDumpProgrammaticGraphs;
+extern ANIMNEXTUNCOOKEDONLY_API TAutoConsoleVariable<bool> CVarDumpProgrammaticGraphs;
 
 struct ANIMNEXTUNCOOKEDONLY_API FUtils
 {
@@ -154,9 +154,6 @@ struct ANIMNEXTUNCOOKEDONLY_API FUtils
 	 **/
 	static FRigVMTemplateArgumentType GetRigVMArgTypeFromParamType(const FAnimNextParamType& InParamType);
 
-	/** Set up a simple animation graph */
-	static void SetupAnimGraph(const FName EntryName, URigVMController* InController);
-	
 	/** Set up a simple event graph */
 	static void SetupEventGraph(URigVMController* InController, UScriptStruct* InEventStruct);
 
diff --git a/Engine/Plugins/Experimental/AnimNext/Source/AnimNextUncookedOnly/Private/AnimNextAssetWorkspaceAssetUserData.cpp b/Engine/Plugins/Experimental/AnimNext/Source/AnimNextUncookedOnly/Private/AnimNextAssetWorkspaceAssetUserData.cpp
--- a/Engine/Plugins/Experimental/AnimNext/Source/AnimNextUncookedOnly/Private/AnimNextAssetWorkspaceAssetUserData.cpp
+++ b/Engine/Plugins/Experimental/AnimNext/Source/AnimNextUncookedOnly/Private/AnimNextAssetWorkspaceAssetUserData.cpp
@@ -23,10 +23,6 @@ void UAnimNextAssetWorkspaceAssetUserData::GetAssetRegistryTags(FAssetRegistryTa
 			{
 				RootAssetExport.GetData().InitializeAsScriptStruct(FAnimNextModuleOutlinerData::StaticStruct());
 			}
-			else if(UAnimNextAnimationGraph* AnimationGraph = Cast<UAnimNextAnimationGraph>(Asset))
-			{
-				RootAssetExport.GetData().InitializeAsScriptStruct(FAnimNextAnimationGraphOutlinerData::StaticStruct());
-			}
 			else if(UAnimNextDataInterface* DataInterface = Cast<UAnimNextDataInterface>(Asset))
 			{
 				RootAssetExport.GetData().InitializeAsScriptStruct(FAnimNextDataInterfaceOutlinerData::StaticStruct());
