AnimNext: Fix CIS issue

Non-unity build

diff --git a/Engine/Plugins/Experimental/AnimNext/Source/AnimNextUncookedOnly/Private/Variables/UniversalObjectLocatorBindingType.cpp b/Engine/Plugins/Experimental/AnimNext/Source/AnimNextUncookedOnly/Private/Variables/UniversalObjectLocatorBindingType.cpp
--- a/Engine/Plugins/Experimental/AnimNext/Source/AnimNextUncookedOnly/Private/Variables/UniversalObjectLocatorBindingType.cpp
+++ b/Engine/Plugins/Experimental/AnimNext/Source/AnimNextUncookedOnly/Private/Variables/UniversalObjectLocatorBindingType.cpp
@@ -14,6 +14,7 @@
 #include "Framework/Application/SlateApplication.h"
 #include "Framework/PropertyViewer/IFieldExpander.h"
 #include "Framework/PropertyViewer/IFieldIterator.h"
+#include "Kismet/BlueprintFunctionLibrary.h"
 #include "Module/RigUnit_AnimNextModuleEvents.h"
 #include "Modules/ModuleManager.h"
 #include "Param/ParamCompatibility.h"
diff --git a/Engine/Plugins/Experimental/AnimNextAnimGraph/Source/AnimNextAnimGraph/Private/Traits/BlendStackTrait.cpp b/Engine/Plugins/Experimental/AnimNextAnimGraph/Source/AnimNextAnimGraph/Private/Traits/BlendStackTrait.cpp
--- a/Engine/Plugins/Experimental/AnimNextAnimGraph/Source/AnimNextAnimGraph/Private/Traits/BlendStackTrait.cpp
+++ b/Engine/Plugins/Experimental/AnimNextAnimGraph/Source/AnimNextAnimGraph/Private/Traits/BlendStackTrait.cpp
@@ -5,6 +5,7 @@
 #include "AnimNextAnimGraphSettings.h"
 #include "AnimNextDataInterfacePayload.h"
 #include "DataInterface/DataInterfaceStructAdapter.h"
+#include "Graph/AnimNextGraphInstance.h"
 #include "TraitCore/NodeInstance.h"
 #include "TraitCore/ExecutionContext.h"
 #include "TraitInterfaces/IGraphFactory.h"
