[Engine]
* Made PrimitiveSceneDesc::GetWorld deprecated (and removed it when WITH_STATE_STREAM is set)

diff --git a/Engine/Plugins/Experimental/AnimNext/Source/AnimNext/Private/AnimNextDebugDraw.cpp b/Engine/Plugins/Experimental/AnimNext/Source/AnimNext/Private/AnimNextDebugDraw.cpp
--- a/Engine/Plugins/Experimental/AnimNext/Source/AnimNext/Private/AnimNextDebugDraw.cpp
+++ b/Engine/Plugins/Experimental/AnimNext/Source/AnimNext/Private/AnimNextDebugDraw.cpp
@@ -1,4 +1,4 @@
-﻿// Copyright Epic Games, Inc. All Rights Reserved.
+// Copyright Epic Games, Inc. All Rights Reserved.
 
 #include "AnimNextDebugDraw.h"
 
@@ -140,14 +140,13 @@ FDebugDraw::FDebugDraw(UObject* InOwner)
 	SceneInfoData->OwnerLastRenderTimePtr = nullptr;
 	SceneDesc.ProxyDesc = &SceneProxyDesc;
 	SceneDesc.PrimitiveSceneData = SceneInfoData;
-	SceneDesc.World = InOwner->GetWorld();
 
 	SceneDesc.RenderMatrix = FMatrix::Identity;
 	SceneDesc.AttachmentRootPosition = FVector::ZeroVector;
 	SceneDesc.Bounds = FBoxSphereBounds(FVector::ZeroVector, FVector::OneVector, 1.f);
 	SceneDesc.LocalBounds = SceneDesc.Bounds;
 
-	Scene = SceneDesc.World->Scene;
+	Scene = InOwner->GetWorld()->Scene;
 
 	Scene->AddPrimitive(&SceneDesc);
 }
diff --git a/Engine/Source/Runtime/Engine/Private/PrimitiveSceneProxy.cpp b/Engine/Source/Runtime/Engine/Private/PrimitiveSceneProxy.cpp
--- a/Engine/Source/Runtime/Engine/Private/PrimitiveSceneProxy.cpp
+++ b/Engine/Source/Runtime/Engine/Private/PrimitiveSceneProxy.cpp
@@ -371,7 +371,10 @@ void FPrimitiveSceneProxyDesc::InitializeFromPrimitiveComponent(const UPrimitive
 	Component = const_cast<UPrimitiveComponent*>(InComponent); 
 	Owner = InComponent->GetOwner();
 
+	#if !WITH_STATE_STREAM
 	World = InComponent->GetWorld();
+	#endif
+
 	CustomPrimitiveData = &InComponent->GetCustomPrimitiveData();
 	Scene = InComponent->GetScene();
 	PrimitiveComponentInterface = InComponent->GetPrimitiveComponentInterface();
@@ -557,6 +560,10 @@ FPrimitiveSceneProxy::FPrimitiveSceneProxy(const FPrimitiveSceneProxyDesc& InPro
 		{
 			SetPrimitiveColor(FActorPrimitiveColorHandler::Get().GetPrimitiveColor(PrimitiveComponent));
 		}
+		else
+		{
+			SetPrimitiveColor(FLinearColor::White);
+		}
 #endif
 	}
 
diff --git a/Engine/Source/Runtime/Engine/Public/PrimitiveSceneDesc.h b/Engine/Source/Runtime/Engine/Public/PrimitiveSceneDesc.h
--- a/Engine/Source/Runtime/Engine/Public/PrimitiveSceneDesc.h
+++ b/Engine/Source/Runtime/Engine/Public/PrimitiveSceneDesc.h
@@ -44,7 +44,6 @@ struct FPrimitiveSceneDesc
 	FMatrix RenderMatrix;
 	FVector AttachmentRootPosition;
 
-	UWorld*  World = nullptr;
 	UObject* PrimitiveUObject = nullptr;
 	
 	// @todo: Possibly add uninitialized FStrings in this object to allow overriding the name without having a corresponding UObject
@@ -71,7 +70,12 @@ struct FPrimitiveSceneDesc
 	EComponentMobility::Type GetMobility() { return Mobility; }
 	FMatrix GetRenderMatrix() { return RenderMatrix; }
 	FVector GetActorPositionForRenderer() { return AttachmentRootPosition; }
+	
+	#if !WITH_STATE_STREAM
+	UE_DEPRECATED(5.6, "World should not be used by rendering")
 	UWorld* GetWorld() { return World; }
+	UWorld*  World = nullptr;
+	#endif
 
 	FBoxSphereBounds GetBounds() { return Bounds; }
 	FBoxSphereBounds GetLocalBounds() { return LocalBounds; }
