AnimNext module garbage collection registering fix, enabling traits to use the garbage collection in the instance data.

diff --git a/Engine/Plugins/Experimental/AnimNext/Source/AnimNext/Internal/Module/AnimNextModuleInstance.h b/Engine/Plugins/Experimental/AnimNext/Source/AnimNext/Internal/Module/AnimNextModuleInstance.h
--- a/Engine/Plugins/Experimental/AnimNext/Source/AnimNext/Internal/Module/AnimNextModuleInstance.h
+++ b/Engine/Plugins/Experimental/AnimNext/Source/AnimNext/Internal/Module/AnimNextModuleInstance.h
@@ -12,6 +12,8 @@
 #include "TraitCore/TraitEventList.h"
 #include "Module/ModuleHandle.h"
 #include "Module/ModuleInstanceComponent.h"
+#include "StructUtils/InstancedStruct.h"
+
 #include "AnimNextModuleInstance.generated.h"
 
 struct FAnimNextModuleInstance;
@@ -52,7 +54,7 @@ namespace UE::AnimNext::Private
 
 namespace UE::AnimNext
 {
-	using ModuleInstanceComponentMapType = TMap<FName, TSharedPtr<FModuleInstanceComponent>>;
+	using ModuleInstanceComponentMapType = TMap<FName, TInstancedStruct<FModuleInstanceComponent>>;
 }
 
 // Root memory owner of a parameterized schedule 
@@ -111,10 +113,8 @@ struct FAnimNextModuleInstance : public FAnimNextDataInterfaceInstance
 	const ComponentType* TryGetComponent() const;
 
 	// Returns a pointer to the specified component, or nullptr if not found
-	ANIMNEXT_API UE::AnimNext::FModuleInstanceComponent* TryGetComponent(int32 ComponentNameHash, FName ComponentName) const;
-
-	// Adds the specified component and returns a reference to it
-	ANIMNEXT_API UE::AnimNext::FModuleInstanceComponent& AddComponent(int32 ComponentNameHash, FName ComponentName, TSharedPtr<UE::AnimNext::FModuleInstanceComponent>&& Component);
+	ANIMNEXT_API const FModuleInstanceComponent* TryGetComponent(int32 ComponentNameHash, FName ComponentName) const;
+	ANIMNEXT_API FModuleInstanceComponent* TryGetComponent(int32 ComponentNameHash, FName ComponentName);
 
 #if UE_ENABLE_DEBUG_DRAWING
 	// Get the debug draw interface
@@ -163,7 +163,8 @@ struct FAnimNextModuleInstance : public FAnimNextDataInterfaceInstance
 #endif
 
 	// Components of this module
-	UE::AnimNext::ModuleInstanceComponentMapType Components; 
+	UPROPERTY()
+	TMap<FName, TInstancedStruct<FModuleInstanceComponent>> ComponentMap;
 
 	// Initialize event index (INDEX_NONE if no initialize event is found)
 	int32 InitializeEventIndex = INDEX_NONE;
@@ -199,6 +200,10 @@ struct FAnimNextModuleInstance : public FAnimNextDataInterfaceInstance
 	// Whether we are currently recreating this instance because of compilation/reinstancing
 	bool bIsRecreatingOnCompile : 1 = false;
 #endif
+
+private:
+	// Adds the specified component and returns a reference to it
+	ANIMNEXT_API FModuleInstanceComponent& AddComponentInternal(int32 ComponentNameHash, FName ComponentName, TInstancedStruct<FModuleInstanceComponent>&& Component);
 };
 
 template<>
@@ -210,18 +215,18 @@ struct TStructOpsTypeTraits<FAnimNextModuleInstance> : public TStructOpsTypeTrai
 	};
 };
 
-template<class ComponentType>
+template<typename ComponentType>
 ComponentType& FAnimNextModuleInstance::GetComponent()
 {
 	const FName ComponentName = ComponentType::StaticComponentName();
 	const int32 ComponentNameHash = GetTypeHash(ComponentName);
 
-	if (UE::AnimNext::FModuleInstanceComponent* Component = TryGetComponent(ComponentNameHash, ComponentName))
+	if (ComponentType* Component = static_cast<ComponentType*>(TryGetComponent(ComponentNameHash, ComponentName)))
 	{
-		return *static_cast<ComponentType*>(Component);
+		return *Component;
 	}
 
-	return static_cast<ComponentType&>(AddComponent(ComponentNameHash, ComponentName, MakeShared<ComponentType>(*this)));
+	return static_cast<ComponentType&>(AddComponentInternal(ComponentNameHash, ComponentName, TInstancedStruct<ComponentType>::Make(*this)));
 }
 
 template<class ComponentType>
diff --git a/Engine/Plugins/Experimental/AnimNext/Source/AnimNext/Internal/Module/ModuleInstanceComponent.h b/Engine/Plugins/Experimental/AnimNext/Source/AnimNext/Internal/Module/ModuleInstanceComponent.h
--- a/Engine/Plugins/Experimental/AnimNext/Source/AnimNext/Internal/Module/ModuleInstanceComponent.h
+++ b/Engine/Plugins/Experimental/AnimNext/Source/AnimNext/Internal/Module/ModuleInstanceComponent.h
@@ -1,8 +1,9 @@
-﻿// Copyright Epic Games, Inc. All Rights Reserved.
+// Copyright Epic Games, Inc. All Rights Reserved.
 
 #pragma once
 
 #include "CoreMinimal.h"
+#include "ModuleInstanceComponent.generated.h"
 
 struct FAnimNextModuleInstance;
 
@@ -11,32 +12,34 @@ struct FAnimNextModuleInstance;
 	static FName StaticComponentName() { return FName(TEXT(#ComponentType)); } \
 	virtual FName GetComponentName() const override { return StaticComponentName(); }
 
-namespace UE::AnimNext
-{
-
 /**
  * FModuleInstanceComponent
  *
  * A module instance component is attached and owned by a module instance.
  * It persists as long as it is needed.
  */
+ USTRUCT()
 struct FModuleInstanceComponent
 {
-	explicit FModuleInstanceComponent(FAnimNextModuleInstance& InOwnerInstance) : OwnerInstance(InOwnerInstance) {}
+	GENERATED_BODY()
+
+	FModuleInstanceComponent() = default;
+	explicit FModuleInstanceComponent(FAnimNextModuleInstance& InOwnerInstance) : OwnerInstance(&InOwnerInstance) {}
 	virtual ~FModuleInstanceComponent() {}
 
 	static FName StaticComponentName() { return FName(TEXT("FModuleInstanceComponent")); }
 	virtual FName GetComponentName() const { return StaticComponentName(); }
 
 	// Returns the owning module instance this component lives on
-	FAnimNextModuleInstance& GetModuleInstance() { return OwnerInstance; }
+	FAnimNextModuleInstance& GetModuleInstance() { return *OwnerInstance; }
 
 	// Returns the owning module instance this component lives on
-	const FAnimNextModuleInstance& GetModuleInstance() const { return OwnerInstance; }
+	const FAnimNextModuleInstance& GetModuleInstance() const { return *OwnerInstance; }
 
 private:
+
 	// The owning module instance this component lives on
-	FAnimNextModuleInstance& OwnerInstance;
-};
+	FAnimNextModuleInstance* OwnerInstance = nullptr;
 
-}
+	friend struct FAnimNextModuleInstance;
+};
diff --git a/Engine/Plugins/Experimental/AnimNext/Source/AnimNext/Private/Component/AnimNextWorldSubsystem.cpp b/Engine/Plugins/Experimental/AnimNext/Source/AnimNext/Private/Component/AnimNextWorldSubsystem.cpp
--- a/Engine/Plugins/Experimental/AnimNext/Source/AnimNext/Private/Component/AnimNextWorldSubsystem.cpp
+++ b/Engine/Plugins/Experimental/AnimNext/Source/AnimNext/Private/Component/AnimNextWorldSubsystem.cpp
@@ -70,7 +70,7 @@ void UAnimNextWorldSubsystem::AddReferencedObjects(UObject* InThis, FReferenceCo
 	UAnimNextWorldSubsystem* This = CastChecked<UAnimNextWorldSubsystem>(InThis);
 	for(FAnimNextModuleInstance& Instance : This->Instances)
 	{
-		Collector.AddPropertyReferences(FAnimNextModuleInstance::StaticStruct(), &Instance, InThis);
+		Collector.AddPropertyReferencesWithStructARO(FAnimNextModuleInstance::StaticStruct(), &Instance, InThis);
 	}
 }
 
diff --git a/Engine/Plugins/Experimental/AnimNext/Source/AnimNext/Private/Module/AnimNextModuleInstance.cpp b/Engine/Plugins/Experimental/AnimNext/Source/AnimNext/Private/Module/AnimNextModuleInstance.cpp
--- a/Engine/Plugins/Experimental/AnimNext/Source/AnimNext/Private/Module/AnimNextModuleInstance.cpp
+++ b/Engine/Plugins/Experimental/AnimNext/Source/AnimNext/Private/Module/AnimNextModuleInstance.cpp
@@ -383,19 +383,33 @@ void FAnimNextModuleInstance::OnModuleCompiled()
 }
 #endif
 
-UE::AnimNext::FModuleInstanceComponent* FAnimNextModuleInstance::TryGetComponent(int32 ComponentNameHash, FName ComponentName) const
+const FModuleInstanceComponent* FAnimNextModuleInstance::TryGetComponent(int32 ComponentNameHash, FName ComponentName) const
 {
-	if (const TSharedPtr<UE::AnimNext::FModuleInstanceComponent>* Component = Components.FindByHash(ComponentNameHash, ComponentName))
+	if (const TInstancedStruct<FModuleInstanceComponent>* InstancedComponent = ComponentMap.FindByHash(ComponentNameHash, ComponentName))
 	{
-		return Component->Get();
+		return InstancedComponent->GetPtr();
+	}
+
+	return nullptr;
+}
+FModuleInstanceComponent* FAnimNextModuleInstance::TryGetComponent(int32 ComponentNameHash, FName ComponentName) 
+{
+	if (TInstancedStruct<FModuleInstanceComponent>* InstancedComponent = ComponentMap.FindByHash(ComponentNameHash, ComponentName))
+	{
+		return InstancedComponent->GetMutablePtr();
 	}
 
 	return nullptr;
 }
 
-UE::AnimNext::FModuleInstanceComponent& FAnimNextModuleInstance::AddComponent(int32 ComponentNameHash, FName ComponentName, TSharedPtr<UE::AnimNext::FModuleInstanceComponent>&& Component)
+FModuleInstanceComponent& FAnimNextModuleInstance::AddComponentInternal(int32 ComponentNameHash, FName ComponentName, TInstancedStruct<FModuleInstanceComponent>&& Component)
 {
-	return *Components.AddByHash(ComponentNameHash, ComponentName, MoveTemp(Component)).Get();
+	TInstancedStruct<FModuleInstanceComponent>& InstancedComponent = ComponentMap.AddByHash(ComponentNameHash, ComponentName, Component);
+	
+	FModuleInstanceComponent* ModuleInstanceComponent = InstancedComponent.GetMutablePtr();
+	check(ModuleInstanceComponent != nullptr);
+
+	return *ModuleInstanceComponent;
 }
 
 #if UE_ENABLE_DEBUG_DRAWING
diff --git a/Engine/Plugins/Experimental/AnimNextAnimGraph/Source/AnimNextAnimGraph/Internal/Graph/AnimNextAnimationGraph.h b/Engine/Plugins/Experimental/AnimNextAnimGraph/Source/AnimNextAnimGraph/Internal/Graph/AnimNextAnimationGraph.h
--- a/Engine/Plugins/Experimental/AnimNextAnimGraph/Source/AnimNextAnimGraph/Internal/Graph/AnimNextAnimationGraph.h
+++ b/Engine/Plugins/Experimental/AnimNextAnimGraph/Source/AnimNextAnimGraph/Internal/Graph/AnimNextAnimationGraph.h
@@ -1,4 +1,4 @@
-﻿// Copyright Epic Games, Inc. All Rights Reserved.
+// Copyright Epic Games, Inc. All Rights Reserved.
 
 #pragma once
 
@@ -25,6 +25,7 @@ struct FAnimNextScheduleGraphTask;
 struct FAnimNextEditorParam;
 struct FAnimNextParam;
 struct FAnimNextModuleInstance;
+struct FAnimationGraphComponent;
 
 namespace UE::AnimNext
 {
@@ -34,7 +35,6 @@ namespace UE::AnimNext
 	struct FTestUtils;
 	struct FParametersProxy;
 	struct FPlayAnimSlotTrait;
-	struct FAnimationGraphComponent;
 }
 
 namespace UE::AnimNext::UncookedOnly
@@ -105,7 +105,7 @@ private:
 	friend class UE::AnimNext::Editor::FVariableCustomization;
 	friend struct UE::AnimNext::FParametersProxy;
 	friend struct UE::AnimNext::FPlayAnimSlotTrait;
-	friend struct UE::AnimNext::FAnimationGraphComponent;
+	friend struct FAnimationGraphComponent;
 
 protected:
 #if WITH_EDITORONLY_DATA
diff --git a/Engine/Plugins/Experimental/AnimNextAnimGraph/Source/AnimNextAnimGraph/Internal/Graph/AnimationGraphComponent.h b/Engine/Plugins/Experimental/AnimNextAnimGraph/Source/AnimNextAnimGraph/Internal/Graph/AnimationGraphComponent.h
--- a/Engine/Plugins/Experimental/AnimNextAnimGraph/Source/AnimNextAnimGraph/Internal/Graph/AnimationGraphComponent.h
+++ b/Engine/Plugins/Experimental/AnimNextAnimGraph/Source/AnimNextAnimGraph/Internal/Graph/AnimationGraphComponent.h
@@ -1,27 +1,38 @@
-﻿// Copyright Epic Games, Inc. All Rights Reserved.
+// Copyright Epic Games, Inc. All Rights Reserved.
 
 #pragma once
 
 #include "Graph/AnimNextGraphInstancePtr.h"
 #include "Module/ModuleInstanceComponent.h"
-
-namespace UE::AnimNext
-{
+#include "AnimationGraphComponent.generated.h"
 
 // Module component that owns/allocates/release all animation graph instances on a module
+USTRUCT()
 struct FAnimationGraphComponent : public FModuleInstanceComponent
 {
+	GENERATED_BODY()
+
 	ANIMNEXT_DECLARE_MODULE_INSTANCE_COMPONENT(FAnimationGraphComponent);
 
+	FAnimationGraphComponent() = default;
 	explicit FAnimationGraphComponent(FAnimNextModuleInstance& InOwnerInstance);
-
+	
 	FAnimNextGraphInstancePtr AllocateInstance(UAnimNextAnimationGraph* InAnimationGraph, FAnimNextGraphInstance* InParentInstance = nullptr, FName InEntryPoint = NAME_None);
 
 	void ReleaseInstance(FAnimNextGraphInstancePtr InInstance);
 
+	void AddStructReferencedObjects(class FReferenceCollector& Collector);
+
 private:
 	// All the owned graph instances for the module
 	TArray<TSharedPtr<FAnimNextGraphInstance>> GraphInstances;
 };
 
-}
+template<>
+struct TStructOpsTypeTraits<FAnimationGraphComponent> : public TStructOpsTypeTraitsBase2<FAnimationGraphComponent>
+{
+	enum
+	{
+		WithAddStructReferencedObjects = true,
+	};
+};
diff --git a/Engine/Plugins/Experimental/AnimNextAnimGraph/Source/AnimNextAnimGraph/Private/Graph/AnimationGraphComponent.cpp b/Engine/Plugins/Experimental/AnimNextAnimGraph/Source/AnimNextAnimGraph/Private/Graph/AnimationGraphComponent.cpp
--- a/Engine/Plugins/Experimental/AnimNextAnimGraph/Source/AnimNextAnimGraph/Private/Graph/AnimationGraphComponent.cpp
+++ b/Engine/Plugins/Experimental/AnimNextAnimGraph/Source/AnimNextAnimGraph/Private/Graph/AnimationGraphComponent.cpp
@@ -1,12 +1,10 @@
-﻿// Copyright Epic Games, Inc. All Rights Reserved.
+// Copyright Epic Games, Inc. All Rights Reserved.
 
 #include "Graph/AnimationGraphComponent.h"
 
+#include "Graph/AnimNextGraphInstance.h"
 #include "Graph/AnimNextAnimationGraph.h"
 
-namespace UE::AnimNext
-{
-
 FAnimationGraphComponent::FAnimationGraphComponent(FAnimNextModuleInstance& InOwnerInstance)
 	: FModuleInstanceComponent(InOwnerInstance)
 {
@@ -37,4 +35,11 @@ void FAnimationGraphComponent::ReleaseInstance(FAnimNextGraphInstancePtr InInsta
 	}
 }
 
+void FAnimationGraphComponent::AddStructReferencedObjects(class FReferenceCollector& Collector)
+{
+	for (const TSharedPtr<FAnimNextGraphInstance>& GraphInstance : GraphInstances)
+	{
+		GraphInstance->AddStructReferencedObjects(Collector);
+	}
 }
+
diff --git a/Engine/Plugins/Experimental/AnimNextAnimGraph/Source/AnimNextAnimGraph/Public/Graph/AnimNextGraphInstancePtr.h b/Engine/Plugins/Experimental/AnimNextAnimGraph/Source/AnimNextAnimGraph/Public/Graph/AnimNextGraphInstancePtr.h
--- a/Engine/Plugins/Experimental/AnimNextAnimGraph/Source/AnimNextAnimGraph/Public/Graph/AnimNextGraphInstancePtr.h
+++ b/Engine/Plugins/Experimental/AnimNextAnimGraph/Source/AnimNextAnimGraph/Public/Graph/AnimNextGraphInstancePtr.h
@@ -17,13 +17,13 @@ class UAnimNextAnimationGraph;
 class FRigVMTraitScope;
 struct FRigVMExtendedExecuteContext;
 struct FAnimNextExecuteContext;
+struct FAnimationGraphComponent;
 
 namespace UE::AnimNext
 {
 	struct FExecutionContext;
 	struct FGraphInstanceComponent;
 	class IDataInterfaceHost;
-	struct FAnimationGraphComponent;
 }
 
 using GraphInstanceComponentMapType = TMap<FName, TSharedPtr<UE::AnimNext::FGraphInstanceComponent>>;
@@ -111,7 +111,7 @@ private:
 	TWeakPtr<FAnimNextGraphInstance> Impl;
 
 	friend UAnimNextAnimationGraph;				// The graph is the one that allocates instances
-	friend UE::AnimNext::FAnimationGraphComponent; // The module component owns the instances
+	friend FAnimationGraphComponent;			// The module component owns the instances
 	friend FRigUnit_AnimNextGraphEvaluator;		// We evaluate the instance
 	friend FRigUnit_AnimNextRunAnimationGraph;	// We evaluate the instance
 	friend UE::AnimNext::FExecutionContext;
