AnimNext: Fix crash when using a newly-created module

We werent correctly handling the case when instancing a module with no variables

diff --git a/Engine/Plugins/Experimental/AnimNext/Source/AnimNext/Private/Module/AnimNextModuleInstance.cpp b/Engine/Plugins/Experimental/AnimNext/Source/AnimNext/Private/Module/AnimNextModuleInstance.cpp
--- a/Engine/Plugins/Experimental/AnimNext/Source/AnimNext/Private/Module/AnimNextModuleInstance.cpp
+++ b/Engine/Plugins/Experimental/AnimNext/Source/AnimNext/Private/Module/AnimNextModuleInstance.cpp
@@ -195,17 +195,20 @@ void FAnimNextModuleInstance::Initialize()
 		// Initialize the RigVM context
 		ExtendedExecuteContext = Module->GetRigVMExtendedExecuteContext();
 
-		// Setup external variables memory ptrs manually as we dont follow the pattern of owning multiple URigVMHosts like control rig.
-		// InitializeVM() is called, but only sets up handles for the defaults in the module, not for an instance
-		TArray<FRigVMExternalVariableRuntimeData> ExternalVariableRuntimeData;
-		ExternalVariableRuntimeData.Reserve(NumVariables);
-		TConstArrayView<FPropertyBagPropertyDesc> Descs = Variables.GetPropertyBagStruct()->GetPropertyDescs();
-		uint8* BasePtr = Variables.GetMutableValue().GetMemory();
-		for(int32 VariableIndex = 0; VariableIndex < NumVariables; ++VariableIndex)
+		if(NumVariables > 0)
 		{
-			ExternalVariableRuntimeData.Emplace(Descs[VariableIndex].CachedProperty->ContainerPtrToValuePtr<uint8>(BasePtr));
+			// Setup external variables memory ptrs manually as we dont follow the pattern of owning multiple URigVMHosts like control rig.
+			// InitializeVM() is called, but only sets up handles for the defaults in the module, not for an instance
+			TArray<FRigVMExternalVariableRuntimeData> ExternalVariableRuntimeData;
+			ExternalVariableRuntimeData.Reserve(NumVariables);
+			TConstArrayView<FPropertyBagPropertyDesc> Descs = Variables.GetPropertyBagStruct()->GetPropertyDescs();
+			uint8* BasePtr = Variables.GetMutableValue().GetMemory();
+			for(int32 VariableIndex = 0; VariableIndex < NumVariables; ++VariableIndex)
+			{
+				ExternalVariableRuntimeData.Emplace(Descs[VariableIndex].CachedProperty->ContainerPtrToValuePtr<uint8>(BasePtr));
+			}
+			ExtendedExecuteContext.ExternalVariableRuntimeData = MoveTemp(ExternalVariableRuntimeData);
 		}
-		ExtendedExecuteContext.ExternalVariableRuntimeData = MoveTemp(ExternalVariableRuntimeData);
 
 		// Now initialize the 'instance', cache memory handles etc. in the context
 		VM->InitializeInstance(ExtendedExecuteContext);
