Fix AnimNext Editor crashes
- On right click state tree in outliner
- On delete multiple trait stacks (While trait stack editor is open in same tab well as details)
- On Reload Asset

#tests Editor try repro crashes

diff --git a/Engine/Plugins/Experimental/AnimNext/Source/AnimNextEditor/Private/Common/AnimNextEdGraphNodeCustomization.cpp b/Engine/Plugins/Experimental/AnimNext/Source/AnimNextEditor/Private/Common/AnimNextEdGraphNodeCustomization.cpp
--- a/Engine/Plugins/Experimental/AnimNext/Source/AnimNextEditor/Private/Common/AnimNextEdGraphNodeCustomization.cpp
+++ b/Engine/Plugins/Experimental/AnimNext/Source/AnimNextEditor/Private/Common/AnimNextEdGraphNodeCustomization.cpp
@@ -424,6 +424,11 @@ void FAnimNextEdGraphNodeCustomization::PopulateCategory(IDetailLayoutBuilder& D
 		FAddPropertyParams AddPropertyParams;
 		IDetailPropertyRow* DetailPropertyRow = ParameterCategory.AddExternalStructureProperty(MakeShared<TInstancedPropertyBagStructureDataProvider<FRigVMMemoryStorageStruct>>(RigVMTypeData->MemoryStorages), TemplateModelPinName, EPropertyLocation::Default, AddPropertyParams);
 
+		if (!DetailPropertyRow)
+		{
+			continue;
+		}
+
 		if (TSharedPtr<IPropertyHandle> Handle = DetailPropertyRow->GetPropertyHandle(); Handle.IsValid())
 		{
 			const TWeakPtr<FRigVMNodeDetailsData> RigVMTypeDataWeak = RigVMTypeData.ToWeakPtr();
diff --git a/Engine/Plugins/Experimental/AnimNext/Source/AnimNextUncookedOnly/Private/AnimNextRigVMAssetEditorData.cpp b/Engine/Plugins/Experimental/AnimNext/Source/AnimNextUncookedOnly/Private/AnimNextRigVMAssetEditorData.cpp
--- a/Engine/Plugins/Experimental/AnimNext/Source/AnimNextUncookedOnly/Private/AnimNextRigVMAssetEditorData.cpp
+++ b/Engine/Plugins/Experimental/AnimNext/Source/AnimNextUncookedOnly/Private/AnimNextRigVMAssetEditorData.cpp
@@ -1045,6 +1045,11 @@ UAnimNextRigVMAssetEntry* UAnimNextRigVMAssetEditorData::FindEntry(FName InName)
 
 	const TObjectPtr<UAnimNextRigVMAssetEntry>* FoundEntry = Entries.FindByPredicate([InName](const UAnimNextRigVMAssetEntry* InEntry)
 	{
+		if (!InEntry)
+		{
+			return false;
+		}
+
 		return InEntry->GetEntryName() == InName;
 	});
 
@@ -1471,6 +1476,11 @@ UAnimNextVariableEntry* UAnimNextRigVMAssetEditorData::AddVariable(FName InName,
 	FName NewParameterName = InName;
 	auto DuplicateNamePredicate = [&NewParameterName](const UAnimNextRigVMAssetEntry* InEntry)
 	{
+		if (!InEntry)
+		{
+			return false;
+		}
+
 		return InEntry->GetEntryName() == NewParameterName;
 	};
 
@@ -1992,4 +2002,4 @@ void UAnimNextRigVMAssetEditorData::SetUseExternalPackages(TArrayView<UAnimNextR
 
 #endif
 
-#undef LOCTEXT_NAMESPACE
\ No newline at end of file
+#undef LOCTEXT_NAMESPACE
diff --git a/Engine/Plugins/Experimental/AnimNextAnimGraph/Source/AnimNextAnimGraphUncookedOnly/Private/Graph/AnimNextAnimationGraph_EditorData.cpp b/Engine/Plugins/Experimental/AnimNextAnimGraph/Source/AnimNextAnimGraphUncookedOnly/Private/Graph/AnimNextAnimationGraph_EditorData.cpp
--- a/Engine/Plugins/Experimental/AnimNextAnimGraph/Source/AnimNextAnimGraphUncookedOnly/Private/Graph/AnimNextAnimationGraph_EditorData.cpp
+++ b/Engine/Plugins/Experimental/AnimNextAnimGraph/Source/AnimNextAnimGraphUncookedOnly/Private/Graph/AnimNextAnimationGraph_EditorData.cpp
@@ -695,10 +695,17 @@ bool UAnimNextAnimationGraph_EditorData::CanAddNewEntry(TSubclassOf<UAnimNextRig
 	// Prevent users adding more than one animation graph
 	if(InClass == UAnimNextAnimationGraphEntry::StaticClass())
 	{
-		if(Entries.ContainsByPredicate([](UAnimNextRigVMAssetEntry* InEntry)
+		auto IsAnimNextGraphEntry = [](UAnimNextRigVMAssetEntry* InEntry)
 		{
-			return InEntry->IsA<UAnimNextAnimationGraphEntry>();
-		}))
+			if (InEntry)
+			{
+				return InEntry->IsA<UAnimNextAnimationGraphEntry>();
+			}
+
+			return false;
+		};
+
+		if(Entries.ContainsByPredicate(IsAnimNextGraphEntry))
 		{
 			return false;
 		};
