Add ability to move AnimNext variable binding task to HighPriority group

diff --git a/Engine/Plugins/Experimental/AnimNext/Source/AnimNext/Private/Module/RigUnit_AnimNextModuleEvents.cpp b/Engine/Plugins/Experimental/AnimNext/Source/AnimNext/Private/Module/RigUnit_AnimNextModuleEvents.cpp
--- a/Engine/Plugins/Experimental/AnimNext/Source/AnimNext/Private/Module/RigUnit_AnimNextModuleEvents.cpp
+++ b/Engine/Plugins/Experimental/AnimNext/Source/AnimNext/Private/Module/RigUnit_AnimNextModuleEvents.cpp
@@ -1,4 +1,4 @@
-﻿// Copyright Epic Games, Inc. All Rights Reserved.
+// Copyright Epic Games, Inc. All Rights Reserved.
 
 #include "Module/RigUnit_AnimNextModuleEvents.h"
 
@@ -6,6 +6,16 @@
 #include "Engine/World.h"
 #include "Module/AnimNextModuleInstance.h"
 
+namespace UE::AnimNext::CVars
+{
+	// Expose control over whether or not the AnimNext binding execution should happen before others in the tick group
+	// This can allow them to dispatch animation jobs earlier, helping hide their latency
+	static TAutoConsoleVariable<bool> CVarHighPriorityAnimNextExecuteBindingsTick(
+		TEXT("a.AnimNext.HighPriorityAnimNextExecuteBindingsTick"),
+		false,
+		TEXT("If true, then schedule the AnimNext binding execution in a high priority tick group before other ticks."));
+}
+
 FRigUnit_AnimNextExecuteBindings_Execute()
 {
 }
@@ -23,6 +33,12 @@ UE::AnimNext::FModuleEventBindingFunction FRigUnit_AnimNextExecuteBindings::GetB
 			FTickFunction& TickFunctions = InContext.ModuleInstance.TickFunctions[TickFunctionIndex];
 			TickFunctions.AddPrerequisite(InContext.ModuleInstance.Object, InContext.ModuleInstance.TickFunctions[InContext.EventIndex]);
 		}
+
+		const bool bHighPriority = UE::AnimNext::CVars::CVarHighPriorityAnimNextExecuteBindingsTick.GetValueOnGameThread();
+		if (InTickFunction.bHighPriority != bHighPriority)
+		{
+			InTickFunction.SetPriorityIncludingPrerequisites(bHighPriority);
+		}
 	};
 }
 
