Fixed android nopch/unity compile error

diff --git a/Engine/Plugins/Experimental/AnimNext/Source/AnimNext/Private/Module/AnimNextModuleInstance.cpp b/Engine/Plugins/Experimental/AnimNext/Source/AnimNext/Private/Module/AnimNextModuleInstance.cpp
--- a/Engine/Plugins/Experimental/AnimNext/Source/AnimNext/Private/Module/AnimNextModuleInstance.cpp
+++ b/Engine/Plugins/Experimental/AnimNext/Source/AnimNext/Private/Module/AnimNextModuleInstance.cpp
@@ -7,6 +7,7 @@
 #endif
 #include "Engine/World.h"
 #include "AnimNextStats.h"
+#include "Async/TaskGraphInterfaces.h"
 #include "SceneInterface.h"
 #include "Algo/TopologicalSort.h"
 #include "Logging/StructuredLog.h"
