Fix crash when collapsing rigVM nodes into function, crashes when trying to generate outliner item for it

diff --git a/Engine/Plugins/Experimental/AnimNext/Source/AnimNextUncookedOnly/Private/UncookedOnlyUtils.cpp b/Engine/Plugins/Experimental/AnimNext/Source/AnimNextUncookedOnly/Private/UncookedOnlyUtils.cpp
--- a/Engine/Plugins/Experimental/AnimNext/Source/AnimNextUncookedOnly/Private/UncookedOnlyUtils.cpp
+++ b/Engine/Plugins/Experimental/AnimNext/Source/AnimNextUncookedOnly/Private/UncookedOnlyUtils.cpp
@@ -885,7 +885,7 @@ void FUtils::CreateSubGraphsOutlinerItemsRecursive(const UAnimNextRigVMAssetEdit
 			{
 				if (URigVMTemplateNode* TemplateRigVMNode = Cast<URigVMTemplateNode>(RigVMEdNode->GetModelNode()))
 				{
-					if (TemplateRigVMNode->GetScriptStruct()->IsChildOf(FRigUnit_AnimNextBase::StaticStruct()))
+					if (TemplateRigVMNode->GetScriptStruct() && TemplateRigVMNode->GetScriptStruct()->IsChildOf(FRigUnit_AnimNextBase::StaticStruct()))
 					{
 						for (URigVMPin* ModelPin : TemplateRigVMNode->GetPins())
 						{
