AnimNext: Fix compilation delegate unregistration using 'this' in destructor

diff --git a/Engine/Plugins/Experimental/AnimNext/Source/AnimNext/Internal/Component/AnimNextWorldSubsystem.h b/Engine/Plugins/Experimental/AnimNext/Source/AnimNext/Internal/Component/AnimNextWorldSubsystem.h
--- a/Engine/Plugins/Experimental/AnimNext/Source/AnimNext/Internal/Component/AnimNextWorldSubsystem.h
+++ b/Engine/Plugins/Experimental/AnimNext/Source/AnimNext/Internal/Component/AnimNextWorldSubsystem.h
@@ -133,6 +133,11 @@ protected:
 	FRWLock InstancesLock;
 	FRWLock PendingLock;
 
+#if WITH_EDITOR
+	// Handle used to hook module compilation
+	FDelegateHandle OnModuleCompiledHandle;
+#endif
+
 	// Handle used to hook into pre-world tick
 	FDelegateHandle OnWorldPreActorTickHandle;
 
diff --git a/Engine/Plugins/Experimental/AnimNext/Source/AnimNext/Private/Component/AnimNextWorldSubsystem.cpp b/Engine/Plugins/Experimental/AnimNext/Source/AnimNext/Private/Component/AnimNextWorldSubsystem.cpp
--- a/Engine/Plugins/Experimental/AnimNext/Source/AnimNext/Private/Component/AnimNextWorldSubsystem.cpp
+++ b/Engine/Plugins/Experimental/AnimNext/Source/AnimNext/Private/Component/AnimNextWorldSubsystem.cpp
@@ -19,7 +19,7 @@ UAnimNextWorldSubsystem::UAnimNextWorldSubsystem()
 	if(!HasAnyFlags(RF_ClassDefaultObject))
 	{
 #if WITH_EDITOR
-		UAnimNextModule::OnModuleCompiled().AddUObject(this, &UAnimNextWorldSubsystem::OnModuleCompiled);
+		OnModuleCompiledHandle = UAnimNextModule::OnModuleCompiled().AddUObject(this, &UAnimNextWorldSubsystem::OnModuleCompiled);
 #endif
 
 		// Kick off root task at the start of each world tick
@@ -40,7 +40,7 @@ UAnimNextWorldSubsystem::~UAnimNextWorldSubsystem()
 	if(!HasAnyFlags(RF_ClassDefaultObject))
 	{
 #if WITH_EDITOR
-		UAnimNextModule::OnModuleCompiled().RemoveAll(this);
+		UAnimNextModule::OnModuleCompiled().Remove(OnModuleCompiledHandle);
 #endif
 
 		FWorldDelegates::OnWorldPreActorTick.Remove(OnWorldPreActorTickHandle);
