AnimNext: Enable component based on init method only
#rnx

diff --git a/Engine/Plugins/Experimental/AnimNext/Source/AnimNext/Private/Component/AnimNextComponent.cpp b/Engine/Plugins/Experimental/AnimNext/Source/AnimNext/Private/Component/AnimNextComponent.cpp
--- a/Engine/Plugins/Experimental/AnimNext/Source/AnimNext/Private/Component/AnimNextComponent.cpp
+++ b/Engine/Plugins/Experimental/AnimNext/Source/AnimNext/Private/Component/AnimNextComponent.cpp
@@ -44,7 +44,10 @@ void UAnimNextComponent::BeginPlay()
 {
 	Super::BeginPlay();
 
-	SetEnabled(true);
+	if (InitMethod == EAnimNextModuleInitMethod::InitializeAndRun)
+	{
+		SetEnabled(true);
+	}
 }
 
 void UAnimNextComponent::EndPlay(const EEndPlayReason::Type EndPlayReason)
