AnimNext: Fix UOL bindings not working in projects with game feature plugins

AnimNext module could end up getting loaded after the dispatch factory registration because of plugin load order, so types would not be present

diff --git a/Engine/Plugins/Experimental/AnimNext/Source/AnimNext/Internal/Variables/RigVMDispatch_CallHoistedAccessorFunction.h b/Engine/Plugins/Experimental/AnimNext/Source/AnimNext/Internal/Variables/RigVMDispatch_CallHoistedAccessorFunction.h
--- a/Engine/Plugins/Experimental/AnimNext/Source/AnimNext/Internal/Variables/RigVMDispatch_CallHoistedAccessorFunction.h
+++ b/Engine/Plugins/Experimental/AnimNext/Source/AnimNext/Internal/Variables/RigVMDispatch_CallHoistedAccessorFunction.h
@@ -21,6 +21,7 @@ struct ANIMNEXT_API FRigVMDispatch_CallHoistedAccessorFunction : public FRigVMDi
 private:
 	virtual UScriptStruct* GetExecuteContextStruct() const override { return FAnimNextExecuteContext::StaticStruct(); }
 	virtual FName GetArgumentNameForOperandIndex(int32 InOperandIndex, int32 InTotalOperands) const override;
+	virtual void RegisterDependencyTypes_NoLock(FRigVMRegistry_NoLock& InRegistry) const override;
 	virtual const TArray<FRigVMTemplateArgumentInfo>& GetArgumentInfos() const override;
 	virtual FRigVMTemplateTypeMap OnNewArgumentType(const FName& InArgumentName, TRigVMTypeIndex InTypeIndex) const override;
 	virtual bool IsSingleton() const override { return true; }
diff --git a/Engine/Plugins/Experimental/AnimNext/Source/AnimNext/Internal/Variables/RigVMDispatch_CallObjectAccessorFunction.h b/Engine/Plugins/Experimental/AnimNext/Source/AnimNext/Internal/Variables/RigVMDispatch_CallObjectAccessorFunction.h
--- a/Engine/Plugins/Experimental/AnimNext/Source/AnimNext/Internal/Variables/RigVMDispatch_CallObjectAccessorFunction.h
+++ b/Engine/Plugins/Experimental/AnimNext/Source/AnimNext/Internal/Variables/RigVMDispatch_CallObjectAccessorFunction.h
@@ -21,6 +21,7 @@ struct ANIMNEXT_API FRigVMDispatch_CallObjectAccessorFunction : public FRigVMDis
 private:
 	virtual UScriptStruct* GetExecuteContextStruct() const override { return FAnimNextExecuteContext::StaticStruct(); }
 	virtual FName GetArgumentNameForOperandIndex(int32 InOperandIndex, int32 InTotalOperands) const override;
+	virtual void RegisterDependencyTypes_NoLock(FRigVMRegistry_NoLock& InRegistry) const override;
 	virtual const TArray<FRigVMTemplateArgumentInfo>& GetArgumentInfos() const override;
 	virtual FRigVMTemplateTypeMap OnNewArgumentType(const FName& InArgumentName, TRigVMTypeIndex InTypeIndex) const override;
 	virtual bool IsSingleton() const override { return true; }
diff --git a/Engine/Plugins/Experimental/AnimNext/Source/AnimNext/Internal/Variables/RigVMDispatch_GetObjectProperty.h b/Engine/Plugins/Experimental/AnimNext/Source/AnimNext/Internal/Variables/RigVMDispatch_GetObjectProperty.h
--- a/Engine/Plugins/Experimental/AnimNext/Source/AnimNext/Internal/Variables/RigVMDispatch_GetObjectProperty.h
+++ b/Engine/Plugins/Experimental/AnimNext/Source/AnimNext/Internal/Variables/RigVMDispatch_GetObjectProperty.h
@@ -21,6 +21,7 @@ struct ANIMNEXT_API FRigVMDispatch_GetObjectProperty : public FRigVMDispatchFact
 private:
 	virtual UScriptStruct* GetExecuteContextStruct() const override { return FAnimNextExecuteContext::StaticStruct(); }
 	virtual FName GetArgumentNameForOperandIndex(int32 InOperandIndex, int32 InTotalOperands) const override;
+	virtual void RegisterDependencyTypes_NoLock(FRigVMRegistry_NoLock& InRegistry) const override;
 	virtual const TArray<FRigVMTemplateArgumentInfo>& GetArgumentInfos() const override;
 	virtual FRigVMTemplateTypeMap OnNewArgumentType(const FName& InArgumentName, TRigVMTypeIndex InTypeIndex) const override;
 	virtual bool IsSingleton() const override { return true; }
diff --git a/Engine/Plugins/Experimental/AnimNext/Source/AnimNext/Private/Variables/RigVMDispatch_CallHoistedAccessorFunction.cpp b/Engine/Plugins/Experimental/AnimNext/Source/AnimNext/Private/Variables/RigVMDispatch_CallHoistedAccessorFunction.cpp
--- a/Engine/Plugins/Experimental/AnimNext/Source/AnimNext/Private/Variables/RigVMDispatch_CallHoistedAccessorFunction.cpp
+++ b/Engine/Plugins/Experimental/AnimNext/Source/AnimNext/Private/Variables/RigVMDispatch_CallHoistedAccessorFunction.cpp
@@ -21,6 +21,32 @@ FName FRigVMDispatch_CallHoistedAccessorFunction::GetArgumentNameForOperandIndex
 	return ArgumentNames[InOperandIndex];
 }
 
+void FRigVMDispatch_CallHoistedAccessorFunction::RegisterDependencyTypes_NoLock(FRigVMRegistry_NoLock& InRegistry) const
+{
+	static UScriptStruct* const AllowedStructTypes[] =
+	{
+		FAnimNextSoftFunctionPtr::StaticStruct(),
+	};
+
+	InRegistry.RegisterStructTypes_NoLock(AllowedStructTypes);
+	for(UScriptStruct* const ScriptStruct : AllowedStructTypes)
+	{
+		InRegistry.FindOrAddType_NoLock(FRigVMTemplateArgumentType(ScriptStruct));
+	}
+
+	static TPair<UClass*, FRigVMRegistry::ERegisterObjectOperation> const AllowedObjectTypes[] =
+	{
+		{ UObject::StaticClass(), FRigVMRegistry::ERegisterObjectOperation::Class },
+	};
+
+	InRegistry.RegisterObjectTypes_NoLock(AllowedObjectTypes);
+
+	for(const TPair<UClass*, FRigVMRegistry::ERegisterObjectOperation>& ObjectType : AllowedObjectTypes)
+	{
+		InRegistry.FindOrAddType_NoLock(FRigVMTemplateArgumentType(ObjectType.Key));
+	}
+}
+
 const TArray<FRigVMTemplateArgumentInfo>& FRigVMDispatch_CallHoistedAccessorFunction::GetArgumentInfos() const
 {
 	static TArray<FRigVMTemplateArgumentInfo> Infos;
diff --git a/Engine/Plugins/Experimental/AnimNext/Source/AnimNext/Private/Variables/RigVMDispatch_CallObjectAccessorFunction.cpp b/Engine/Plugins/Experimental/AnimNext/Source/AnimNext/Private/Variables/RigVMDispatch_CallObjectAccessorFunction.cpp
--- a/Engine/Plugins/Experimental/AnimNext/Source/AnimNext/Private/Variables/RigVMDispatch_CallObjectAccessorFunction.cpp
+++ b/Engine/Plugins/Experimental/AnimNext/Source/AnimNext/Private/Variables/RigVMDispatch_CallObjectAccessorFunction.cpp
@@ -21,6 +21,32 @@ FName FRigVMDispatch_CallObjectAccessorFunction::GetArgumentNameForOperandIndex(
 	return ArgumentNames[InOperandIndex];
 }
 
+void FRigVMDispatch_CallObjectAccessorFunction::RegisterDependencyTypes_NoLock(FRigVMRegistry_NoLock& InRegistry) const
+{
+	static UScriptStruct* const AllowedStructTypes[] =
+	{
+		FAnimNextSoftFunctionPtr::StaticStruct(),
+	};
+
+	InRegistry.RegisterStructTypes_NoLock(AllowedStructTypes);
+	for(UScriptStruct* const ScriptStruct : AllowedStructTypes)
+	{
+		InRegistry.FindOrAddType_NoLock(FRigVMTemplateArgumentType(ScriptStruct));
+	}
+
+	static TPair<UClass*, FRigVMRegistry::ERegisterObjectOperation> const AllowedObjectTypes[] =
+	{
+		{ UObject::StaticClass(), FRigVMRegistry::ERegisterObjectOperation::Class },
+	};
+
+	InRegistry.RegisterObjectTypes_NoLock(AllowedObjectTypes);
+
+	for(const TPair<UClass*, FRigVMRegistry::ERegisterObjectOperation>& ObjectType : AllowedObjectTypes)
+	{
+		InRegistry.FindOrAddType_NoLock(FRigVMTemplateArgumentType(ObjectType.Key));
+	}
+}
+
 const TArray<FRigVMTemplateArgumentInfo>& FRigVMDispatch_CallObjectAccessorFunction::GetArgumentInfos() const
 {
 	static TArray<FRigVMTemplateArgumentInfo> Infos;
diff --git a/Engine/Plugins/Experimental/AnimNext/Source/AnimNext/Private/Variables/RigVMDispatch_GetObjectProperty.cpp b/Engine/Plugins/Experimental/AnimNext/Source/AnimNext/Private/Variables/RigVMDispatch_GetObjectProperty.cpp
--- a/Engine/Plugins/Experimental/AnimNext/Source/AnimNext/Private/Variables/RigVMDispatch_GetObjectProperty.cpp
+++ b/Engine/Plugins/Experimental/AnimNext/Source/AnimNext/Private/Variables/RigVMDispatch_GetObjectProperty.cpp
@@ -21,6 +21,32 @@ FName FRigVMDispatch_GetObjectProperty::GetArgumentNameForOperandIndex(int32 InO
 	return ArgumentNames[InOperandIndex];
 }
 
+void FRigVMDispatch_GetObjectProperty::RegisterDependencyTypes_NoLock(FRigVMRegistry_NoLock& InRegistry) const
+{
+	static UScriptStruct* const AllowedStructTypes[] =
+	{
+		FAnimNextFieldPath::StaticStruct(),
+	};
+
+	InRegistry.RegisterStructTypes_NoLock(AllowedStructTypes);
+	for(UScriptStruct* const ScriptStruct : AllowedStructTypes)
+	{
+		InRegistry.FindOrAddType_NoLock(FRigVMTemplateArgumentType(ScriptStruct));
+	}
+
+	static TPair<UClass*, FRigVMRegistry::ERegisterObjectOperation> const AllowedObjectTypes[] =
+	{
+		{ UObject::StaticClass(), FRigVMRegistry::ERegisterObjectOperation::Class },
+	};
+
+	InRegistry.RegisterObjectTypes_NoLock(AllowedObjectTypes);
+
+	for(const TPair<UClass*, FRigVMRegistry::ERegisterObjectOperation>& ObjectType : AllowedObjectTypes)
+	{
+		InRegistry.FindOrAddType_NoLock(FRigVMTemplateArgumentType(ObjectType.Key));
+	}
+}
+
 const TArray<FRigVMTemplateArgumentInfo>& FRigVMDispatch_GetObjectProperty::GetArgumentInfos() const
 {
 	static TArray<FRigVMTemplateArgumentInfo> Infos;
