AnimNext: Prevent stall on startup in projects with a large asset registry

Skip the external package scan if we are not externally packaged (scan uses the asset registry and can stall if the registry takes a long time to start up).
Remove external packaging for default factory asset(s)

#virtualized

diff --git a/Engine/Plugins/Experimental/AnimNext/Source/AnimNextUncookedOnly/Private/AnimNextRigVMAssetEditorData.cpp b/Engine/Plugins/Experimental/AnimNext/Source/AnimNextUncookedOnly/Private/AnimNextRigVMAssetEditorData.cpp
--- a/Engine/Plugins/Experimental/AnimNext/Source/AnimNextUncookedOnly/Private/AnimNextRigVMAssetEditorData.cpp
+++ b/Engine/Plugins/Experimental/AnimNext/Source/AnimNextUncookedOnly/Private/AnimNextRigVMAssetEditorData.cpp
@@ -217,12 +217,15 @@ void UAnimNextRigVMAssetEditorData::PostLoad()
 
 void UAnimNextRigVMAssetEditorData::PostLoadExternalPackages()
 {
-	FExternalPackageHelper::LoadObjectsFromExternalPackages<UAnimNextRigVMAssetEntry>(this, [this](UAnimNextRigVMAssetEntry* InLoadedEntry)
+	if(bUsesExternalPackages)
 	{
-		check(IsValid(InLoadedEntry));
-		InLoadedEntry->Initialize(this);
-		Entries.Add(InLoadedEntry);
-	});
+		FExternalPackageHelper::LoadObjectsFromExternalPackages<UAnimNextRigVMAssetEntry>(this, [this](UAnimNextRigVMAssetEntry* InLoadedEntry)
+		{
+			check(IsValid(InLoadedEntry));
+			InLoadedEntry->Initialize(this);
+			Entries.Add(InLoadedEntry);
+		});
+	}
 
 	// Internal entries should be empty if we are externally packaged
 	ensure(bUsesExternalPackages == InternalEntries.IsEmpty());
