AnimNext: Added the ability to specify custom tick dependencies from within a module

Added a 'add/remove dependencies' node. This allows various traits to be added to itself via a custom workflow on the node. Each trait is a 'dependency' type that fully encapsulates how to establish/disband the dependency. Dependencies are only configured on the game thread via a GT task that is launched from the node's execute function. Multiple add/remove dependency will not be batched together.
Added a function to get all the available event names for a module so they can be selected from a dropdown on the dependency traits.

diff --git a/Engine/Plugins/Experimental/AnimNext/Source/AnimNext/AnimNext.Build.cs b/Engine/Plugins/Experimental/AnimNext/Source/AnimNext/AnimNext.Build.cs
--- a/Engine/Plugins/Experimental/AnimNext/Source/AnimNext/AnimNext.Build.cs
+++ b/Engine/Plugins/Experimental/AnimNext/Source/AnimNext/AnimNext.Build.cs
@@ -31,7 +31,11 @@ namespace UnrealBuildTool.Rules
 				PrivateDependencyModuleNames.AddRange(
 					new string[]
 					{
-						"RigVMDeveloper"
+						"RigVMDeveloper",
+						"ToolMenus",
+						"Slate",
+						"SlateCore",
+						"UnrealEd"
 					}
 				);
 			}
diff --git a/Engine/Plugins/Experimental/AnimNext/Source/AnimNext/Internal/AnimNextConfig.h b/Engine/Plugins/Experimental/AnimNext/Source/AnimNext/Internal/AnimNextConfig.h
--- a/Engine/Plugins/Experimental/AnimNext/Source/AnimNext/Internal/AnimNextConfig.h
+++ b/Engine/Plugins/Experimental/AnimNext/Source/AnimNext/Internal/AnimNextConfig.h
@@ -14,4 +14,8 @@ private:
 #if WITH_EDITOR
 	virtual void PostEditChangeProperty(FPropertyChangedEvent& PropertyChangedEvent) override;
 #endif
+
+	// Helper function used to display UI options for module events
+	UFUNCTION()
+	static TArray<FString> GetModuleEventNames();
 };
\ No newline at end of file
diff --git a/Engine/Plugins/Experimental/AnimNext/Source/AnimNext/Internal/Module/AnimNextModuleInstance.h b/Engine/Plugins/Experimental/AnimNext/Source/AnimNext/Internal/Module/AnimNextModuleInstance.h
--- a/Engine/Plugins/Experimental/AnimNext/Source/AnimNext/Internal/Module/AnimNextModuleInstance.h
+++ b/Engine/Plugins/Experimental/AnimNext/Source/AnimNext/Internal/Module/AnimNextModuleInstance.h
@@ -130,6 +130,12 @@ struct FAnimNextModuleInstance : public FAnimNextDataInterfaceInstance
 	// Run a simple task on the GT via FFunctionGraphTask::CreateAndDispatchWhenReady
 	static void RunTaskOnGameThread(TUniqueFunction<void(void)>&& InFunction);
 
+	// Get the names of all available module names
+	static void GetModuleEventNames(TArray<FString>& OutNames, bool bIncludeInternal = false);
+
+	// Find the tick function for the specified event
+	UE::AnimNext::FModuleEventTickFunction* FindTickFunction(FName InEventName);
+	
 	// Object this entry is bound to
 	UPROPERTY(Transient)
 	TObjectPtr<UObject> Object = nullptr;
diff --git a/Engine/Plugins/Experimental/AnimNext/Source/AnimNext/Internal/Module/RigUnit_AnimNextModuleEventDependencies.h b/Engine/Plugins/Experimental/AnimNext/Source/AnimNext/Internal/Module/RigUnit_AnimNextModuleEventDependencies.h
new file mode 100644
--- /dev/null
+++ b/Engine/Plugins/Experimental/AnimNext/Source/AnimNext/Internal/Module/RigUnit_AnimNextModuleEventDependencies.h
@@ -0,0 +1,57 @@
+// Copyright Epic Games, Inc. All Rights Reserved.
+
+#pragma once
+
+#include "RigUnit_AnimNextModuleEvents.h"
+#include "RigUnit_AnimNextModuleEventDependencies.generated.h"
+
+class URigVMController;
+class URigVMNode;
+
+UCLASS()
+class UAnimNextAddDependencyMenuContext : public UObject
+{
+	GENERATED_BODY()
+public:
+	TObjectPtr<URigVMController> Controller = nullptr;
+	TObjectPtr<URigVMNode> Node = nullptr;
+};
+
+USTRUCT(meta=(Hidden, Category="Dependencies", NodeColor="0, 1, 1", Keywords="Prerequisite,Subsequent,TickFunction"))
+struct ANIMNEXT_API FRigUnit_AnimNextModuleDependenciesBase : public FRigUnit_AnimNextBase
+{
+	GENERATED_BODY()
+
+	FRigUnit_AnimNextModuleDependenciesBase() = default;
+
+	// FRigVMStruct interface
+	virtual TArray<FRigVMUserWorkflow> GetSupportedWorkflows(const UObject* InSubject) const override;
+
+	// The execution result
+	UPROPERTY(EditAnywhere, DisplayName = "Execute", Category = "Events", meta = (Input, Output))
+	FAnimNextExecuteContext ExecuteContext;
+};
+
+/** Adds execution dependencies between module events and other systems */
+USTRUCT(meta=(DisplayName="Add Dependencies"))
+struct ANIMNEXT_API FRigUnit_AnimNextModuleAddDependencies : public FRigUnit_AnimNextModuleDependenciesBase
+{
+	GENERATED_BODY()
+
+	FRigUnit_AnimNextModuleAddDependencies() = default;
+
+	RIGVM_METHOD()
+	virtual void Execute() override;
+};
+
+/** Removes execution dependencies between module events and some systems */
+USTRUCT(meta=(DisplayName="Remove Dependencies"))
+struct ANIMNEXT_API FRigUnit_AnimNextModuleRemoveDependencies : public FRigUnit_AnimNextModuleDependenciesBase
+{
+	GENERATED_BODY()
+
+	FRigUnit_AnimNextModuleRemoveDependencies() = default;
+
+	RIGVM_METHOD()
+	virtual void Execute() override;
+};
\ No newline at end of file
diff --git a/Engine/Plugins/Experimental/AnimNext/Source/AnimNext/Private/AnimNextConfig.cpp b/Engine/Plugins/Experimental/AnimNext/Source/AnimNext/Private/AnimNextConfig.cpp
--- a/Engine/Plugins/Experimental/AnimNext/Source/AnimNext/Private/AnimNextConfig.cpp
+++ b/Engine/Plugins/Experimental/AnimNext/Source/AnimNext/Private/AnimNextConfig.cpp
@@ -2,6 +2,8 @@
 
 #include "AnimNextConfig.h"
 
+#include "Module/RigUnit_AnimNextModuleEvents.h"
+
 #if WITH_EDITOR
 
 void UAnimNextConfig::PostEditChangeProperty(FPropertyChangedEvent& PropertyChangedEvent)
@@ -9,4 +11,11 @@ void UAnimNextConfig::PostEditChangeProperty(FPropertyChangedEvent& PropertyChan
 	SaveConfig();
 }
 
-#endif
\ No newline at end of file
+#endif
+
+TArray<FString> UAnimNextConfig::GetModuleEventNames()
+{
+	TArray<FString> EventNames;
+	FAnimNextModuleInstance::GetModuleEventNames(EventNames);
+	return EventNames;
+}
\ No newline at end of file
diff --git a/Engine/Plugins/Experimental/AnimNext/Source/AnimNext/Private/Module/AnimNextModuleInstance.cpp b/Engine/Plugins/Experimental/AnimNext/Source/AnimNext/Private/Module/AnimNextModuleInstance.cpp
--- a/Engine/Plugins/Experimental/AnimNext/Source/AnimNext/Private/Module/AnimNextModuleInstance.cpp
+++ b/Engine/Plugins/Experimental/AnimNext/Source/AnimNext/Private/Module/AnimNextModuleInstance.cpp
@@ -451,4 +451,40 @@ void FAnimNextModuleInstance::RunTaskOnGameThread(TUniqueFunction<void(void)>&&
 	{
 		FFunctionGraphTask::CreateAndDispatchWhenReady(MoveTemp(InFunction), TStatId(), nullptr, ENamedThreads::GameThread);
 	}
+}
+
+void FAnimNextModuleInstance::GetModuleEventNames(TArray<FString>& OutNames, bool bIncludeInternal)
+{
+	using namespace UE::AnimNext::Private;
+
+	check(IsInGameThread());			// This function cannot be run concurrently because of static usage
+	check(GAllModuleEvents.Num() > 0);	// Call CacheAllModuleEvents before this function
+	
+	for(const FImplementedModuleEvent& ModuleEvent : GAllModuleEvents)
+	{
+#if WITH_METADATA
+		bool bIsHidden = ModuleEvent.Struct->HasMetaData(FRigVMStruct::HiddenMetaName);
+#else
+		bool bIsHidden = false;
+#endif
+	
+		if(bIncludeInternal || !bIsHidden)
+		{
+			OutNames.Add(ModuleEvent.EventName.ToString());
+		}
+	}
+}
+
+UE::AnimNext::FModuleEventTickFunction* FAnimNextModuleInstance::FindTickFunction(FName InEventName)
+{
+	using namespace UE::AnimNext;
+	
+	for(FModuleEventTickFunction& TickFunction : TickFunctions)
+	{
+		if(TickFunction.EventName == InEventName)
+		{
+			return &TickFunction;
+		}
+	}
+	return nullptr;
 }
\ No newline at end of file
diff --git a/Engine/Plugins/Experimental/AnimNext/Source/AnimNext/Private/Module/RigUnit_AnimNextModuleEventDependencies.cpp b/Engine/Plugins/Experimental/AnimNext/Source/AnimNext/Private/Module/RigUnit_AnimNextModuleEventDependencies.cpp
new file mode 100644
--- /dev/null
+++ b/Engine/Plugins/Experimental/AnimNext/Source/AnimNext/Private/Module/RigUnit_AnimNextModuleEventDependencies.cpp
@@ -0,0 +1,238 @@
+// Copyright Epic Games, Inc. All Rights Reserved.
+
+#include "Module/RigUnit_AnimNextModuleEventDependencies.h"
+#include "Module/RigVMTrait_ModuleEventDependency.h"
+#if WITH_EDITOR
+#include "RigVMModel/RigVMController.h"
+#include "StructViewerModule.h"
+#include "ToolMenus.h"
+#include "Framework/Application/SlateApplication.h"
+#include "ScopedTransaction.h"
+#include "StructViewerFilter.h"
+#include "Widgets/Layout/SBox.h"
+#endif
+
+namespace UE::AnimNext
+{
+
+// Helper function used to add/remove dependencies
+template<typename PredicateType>
+static void ApplyDependenciesHelper(const FAnimNextExecuteContext& InContext, const PredicateType& InPredicate)
+{
+	FAnimNextModuleInstance& ModuleInstance = InContext.GetContextData<FAnimNextModuleContextData>().GetModuleInstance();
+	TArray<const FRigVMTrait_ModuleEventDependency*> Dependencies;
+	Dependencies.Reserve(InContext.GetTraits().Num());
+	for(const FRigVMTraitScope& TraitScope : InContext.GetTraits())
+	{
+		const FRigVMTrait_ModuleEventDependency* DependencyTrait = TraitScope.GetTrait<FRigVMTrait_ModuleEventDependency>();
+		if(DependencyTrait == nullptr)
+		{
+			continue;
+		}
+
+		Dependencies.Add(DependencyTrait);
+	}
+
+	if(Dependencies.Num() == 0)
+	{
+		return;
+	}
+
+	TWeakObjectPtr<UObject> WeakObject = ModuleInstance.Object;
+	TArray<FTickFunction*> TickFunctions;
+	TickFunctions.Reserve(Dependencies.Num());
+	bool bTickFunctionFound = false; 
+	for(const FRigVMTrait_ModuleEventDependency* Dependency : Dependencies)
+	{
+		FTickFunction* TickFunction = ModuleInstance.FindTickFunction(Dependency->EventName);
+		TickFunctions.Add(TickFunction);
+		bTickFunctionFound |= (TickFunction != nullptr);
+	}
+
+	if(!bTickFunctionFound)
+	{
+		return;
+	}
+
+	// Tick function prerequisites can only be updated on the game thread
+	FAnimNextModuleInstance::RunTaskOnGameThread([WeakObject, TickFunctions = MoveTemp(TickFunctions), Dependencies = MoveTemp(Dependencies), InPredicate]()
+	{
+		UObject* Object = WeakObject.Get();
+		if(Object == nullptr)
+		{
+			return;
+		}
+
+		check(TickFunctions.Num() == Dependencies.Num());
+
+		for(int32 DependencyIndex = 0; DependencyIndex < Dependencies.Num(); ++DependencyIndex)
+		{
+			FTickFunction* TickFunction = TickFunctions[DependencyIndex];
+			if(TickFunction == nullptr)
+			{
+				continue;
+			}
+			
+			FModuleDependencyContext Context(Object, *TickFunction);
+			const FRigVMTrait_ModuleEventDependency* Dependency = Dependencies[DependencyIndex];
+			InPredicate(Context, *Dependency);
+		}
+	});
+}
+
+}
+
+TArray<FRigVMUserWorkflow> FRigUnit_AnimNextModuleDependenciesBase::GetSupportedWorkflows(const UObject* InSubject) const
+{
+	TArray<FRigVMUserWorkflow> Workflows = Super::GetSupportedWorkflows(InSubject);
+
+#if WITH_EDITOR
+	Workflows.Emplace(
+		TEXT("Add"),
+		TEXT("Adds a module event dependency to this node"),
+		ERigVMUserWorkflowType::NodeContextButton,
+		FRigVMPerformUserWorkflowDelegate::CreateLambda([](const URigVMUserWorkflowOptions* InOptions, UObject* InController)
+		{
+			URigVMController* Controller = CastChecked<URigVMController>(InController);
+			if(Controller == nullptr)
+			{
+				return false;
+			}
+
+			URigVMNode* Node = InOptions->GetSubject<URigVMNode>();
+			if(Node == nullptr)
+			{
+				return false;
+			}
+
+			static const FName MenuName = "AnimNextModuleDependenciesAddMenu";
+			if(!UToolMenus::Get()->IsMenuRegistered(MenuName))
+			{
+				UToolMenu* Menu = UToolMenus::Get()->RegisterMenu(MenuName);
+				Menu->AddDynamicSection("DependencyTraits", FNewToolMenuDelegate::CreateLambda([](UToolMenu* InMenu)
+				{
+					UAnimNextAddDependencyMenuContext* ContextObject = InMenu->FindContext<UAnimNextAddDependencyMenuContext>();
+					if(ContextObject == nullptr)
+					{
+						return;
+					}
+
+					auto OnTraitPicked = [Controller = ContextObject->Controller, Node = ContextObject->Node](const UScriptStruct* InStruct)
+					{
+						FSlateApplication::Get().DismissAllMenus();
+
+						FScopedTransaction Transaction(NSLOCTEXT("AnimNextModuleDependencies", "AddDependencyTraitTransaction", "Add dependency trait"));
+						Controller->AddTrait(Node, const_cast<UScriptStruct*>(InStruct), InStruct->GetFName(), TEXT(""));
+					};
+
+					class FStructFilter : public IStructViewerFilter
+					{
+					public:
+						virtual bool IsStructAllowed(const FStructViewerInitializationOptions& InInitOptions, const UScriptStruct* InStruct, TSharedRef<FStructViewerFilterFuncs> InFilterFuncs) override
+						{
+							if (InStruct->HasMetaData(TEXT("Hidden")))
+							{
+								return false;
+							}
+
+							return InStruct->IsChildOf(FRigVMTrait_ModuleEventDependency::StaticStruct()) && InStruct != FRigVMTrait_ModuleEventDependency::StaticStruct();
+						}
+
+						virtual bool IsUnloadedStructAllowed(const FStructViewerInitializationOptions& InInitOptions, const FSoftObjectPath& InStructPath, TSharedRef<class FStructViewerFilterFuncs> InFilterFuncs)
+						{
+							return false;
+						};
+					};
+					
+					FToolMenuSection& Section = InMenu->AddSection("DependencyTraits");
+					FStructViewerModule& StructViewerModule = FModuleManager::LoadModuleChecked<FStructViewerModule>("StructViewer");
+					FStructViewerInitializationOptions InitOptions;
+					InitOptions.NameTypeToDisplay = EStructViewerNameTypeToDisplay::DisplayName;
+					InitOptions.StructFilter = MakeShared<FStructFilter>();
+					Section.AddEntry(
+						FToolMenuEntry::InitWidget(
+							"Traits",
+							SNew(SBox)
+							.WidthOverride(300.0f)
+							.HeightOverride(400.0f)
+							[
+								StructViewerModule.CreateStructViewer(InitOptions, FOnStructPicked::CreateLambda(OnTraitPicked))
+							],
+							FText::GetEmpty(),
+							true,
+							false,
+							true));
+				}));
+			}
+
+			UAnimNextAddDependencyMenuContext* ContextObject = NewObject<UAnimNextAddDependencyMenuContext>();
+			ContextObject->Controller = Controller;
+			ContextObject->Node = Node;
+
+			FSlateApplication::Get().PushMenu(
+				FSlateApplication::Get().GetInteractiveTopLevelWindows()[0],
+				FWidgetPath(),
+				UToolMenus::Get()->GenerateWidget(MenuName, FToolMenuContext(ContextObject)),
+				FSlateApplication::Get().GetCursorPos(),
+				FPopupTransitionEffect(FPopupTransitionEffect::ContextMenu)
+			);
+
+			return true;
+		}),
+		URigVMUserWorkflowOptions::StaticClass());
+
+	Workflows.Emplace(
+		TEXT("Remove Dependency"),
+		TEXT("Removes this module event dependency from this node"),
+		ERigVMUserWorkflowType::PinContext,
+		FRigVMPerformUserWorkflowDelegate::CreateLambda([](const URigVMUserWorkflowOptions* InOptions, UObject* InController)
+		{
+			URigVMController* Controller = CastChecked<URigVMController>(InController);
+			if(Controller == nullptr)
+			{
+				return false;
+			}
+
+			URigVMPin* Pin = InOptions->GetSubject<URigVMPin>();
+			if(Pin == nullptr)
+			{
+				return false;
+			}
+
+			if(!Pin->IsTraitPin() || Pin->GetTraitScriptStruct() == nullptr || !Pin->GetTraitScriptStruct()->IsChildOf(FRigVMTrait_ModuleEventDependency::StaticStruct()))
+			{
+				return false;
+			}
+
+			URigVMPin* TraitPin = Pin->GetNode()->FindTrait(Pin);
+			if(TraitPin == nullptr)
+			{
+				return false;
+			}
+
+			FScopedTransaction Transaction(NSLOCTEXT("AnimNextModuleDependencies", "RemoveDependencyTraitTransaction", "Remove dependency trait"));
+			Controller->RemoveTrait(Pin->GetNode(), TraitPin->GetFName());
+
+			return true;
+		}),
+		URigVMUserWorkflowOptions::StaticClass());
+#endif
+
+	return Workflows;
+}
+
+FRigUnit_AnimNextModuleAddDependencies_Execute()
+{
+	UE::AnimNext::ApplyDependenciesHelper(ExecuteContext, [](const UE::AnimNext::FModuleDependencyContext& InContext, const FRigVMTrait_ModuleEventDependency& InDependency)
+	{
+		InDependency.OnAddDependency(InContext);
+	});
+}
+
+FRigUnit_AnimNextModuleRemoveDependencies_Execute()
+{
+	UE::AnimNext::ApplyDependenciesHelper(ExecuteContext, [](const UE::AnimNext::FModuleDependencyContext& InContext, const FRigVMTrait_ModuleEventDependency& InDependency)
+	{
+		InDependency.OnRemoveDependency(InContext);
+	});
+}
\ No newline at end of file
diff --git a/Engine/Plugins/Experimental/AnimNext/Source/AnimNext/Private/Module/RigVMTrait_ModuleEventDependency_ActorComponentClassPrimaryTickFunction.cpp b/Engine/Plugins/Experimental/AnimNext/Source/AnimNext/Private/Module/RigVMTrait_ModuleEventDependency_ActorComponentClassPrimaryTickFunction.cpp
new file mode 100644
--- /dev/null
+++ b/Engine/Plugins/Experimental/AnimNext/Source/AnimNext/Private/Module/RigVMTrait_ModuleEventDependency_ActorComponentClassPrimaryTickFunction.cpp
@@ -0,0 +1,70 @@
+// Copyright Epic Games, Inc. All Rights Reserved.
+
+#include "RigVMTrait_ModuleEventDependency_ActorComponentClassPrimaryTickFunction.h"
+
+#if WITH_EDITOR
+FString FRigVMTrait_ModuleEventDependency_ActorComponentClassPrimaryTickFunction::GetDisplayName() const
+{
+	return StaticStruct()->GetDisplayNameText().ToString();
+}
+#endif
+
+void FRigVMTrait_ModuleEventDependency_ActorComponentClassPrimaryTickFunction::OnAddDependency(const UE::AnimNext::FModuleDependencyContext& InContext) const
+{
+	UActorComponent* AnimNextComponent = Cast<UActorComponent>(InContext.Object);
+	if(AnimNextComponent == nullptr)
+	{
+		return;
+	}
+
+	AActor* Actor = AnimNextComponent->GetOwner();
+	if(Actor == nullptr)
+	{
+		return;
+	}
+
+	UActorComponent* TargetComponent = Actor->GetComponentByClass(ComponentClass);
+	if(TargetComponent == nullptr)
+	{
+		return;
+	}
+
+	if(Ordering == EAnimNextModuleEventDependencyOrdering::Before)
+	{
+		InContext.TickFunction.AddPrerequisite(TargetComponent, TargetComponent->PrimaryComponentTick);
+	}
+	else
+	{
+		TargetComponent->PrimaryComponentTick.AddPrerequisite(AnimNextComponent, InContext.TickFunction);
+	}
+}
+
+void FRigVMTrait_ModuleEventDependency_ActorComponentClassPrimaryTickFunction::OnRemoveDependency(const UE::AnimNext::FModuleDependencyContext& InContext) const
+{
+	UActorComponent* AnimNextComponent = Cast<UActorComponent>(InContext.Object);
+	if(AnimNextComponent == nullptr)
+	{
+		return;
+	}
+
+	AActor* Actor = AnimNextComponent->GetOwner();
+	if(Actor == nullptr)
+	{
+		return;
+	}
+
+	UActorComponent* TargetComponent = Actor->GetComponentByClass(ComponentClass);
+	if(TargetComponent == nullptr)
+	{
+		return;
+	}
+
+	if(Ordering == EAnimNextModuleEventDependencyOrdering::Before)
+	{
+		InContext.TickFunction.RemovePrerequisite(TargetComponent, TargetComponent->PrimaryComponentTick);
+	}
+	else
+	{
+		TargetComponent->PrimaryComponentTick.RemovePrerequisite(AnimNextComponent, InContext.TickFunction);
+	}
+}
\ No newline at end of file
diff --git a/Engine/Plugins/Experimental/AnimNext/Source/AnimNext/Private/Module/RigVMTrait_ModuleEventDependency_ActorComponentClassPrimaryTickFunction.h b/Engine/Plugins/Experimental/AnimNext/Source/AnimNext/Private/Module/RigVMTrait_ModuleEventDependency_ActorComponentClassPrimaryTickFunction.h
new file mode 100644
--- /dev/null
+++ b/Engine/Plugins/Experimental/AnimNext/Source/AnimNext/Private/Module/RigVMTrait_ModuleEventDependency_ActorComponentClassPrimaryTickFunction.h
@@ -0,0 +1,29 @@
+// Copyright Epic Games, Inc. All Rights Reserved.
+
+#pragma once
+
+#include "Module/RigVMTrait_ModuleEventDependency.h"
+#include "Templates/SubclassOf.h"
+#include "RigVMTrait_ModuleEventDependency_ActorComponentClassPrimaryTickFunction.generated.h"
+
+class UActorComponent;
+
+// A dependency on the primary tick function of the first-found actor component of the specified class
+USTRUCT(DisplayName="Actor Component by Class", meta=(ShowTooltip=true))
+struct FRigVMTrait_ModuleEventDependency_ActorComponentClassPrimaryTickFunction : public FRigVMTrait_ModuleEventDependency
+{
+	GENERATED_BODY()
+
+	// FRigVMTrait interface
+#if WITH_EDITOR
+	virtual FString GetDisplayName() const override;
+#endif
+
+	// FAnimNextModuleEventDependency interface
+	virtual void OnAddDependency(const UE::AnimNext::FModuleDependencyContext& InContext) const override;
+	virtual void OnRemoveDependency(const UE::AnimNext::FModuleDependencyContext& InContext) const override;
+
+	// The component class to look for when establishing the dependency
+	UPROPERTY(EditAnywhere, Category = "Dependency")
+	TSubclassOf<UActorComponent> ComponentClass;
+};
\ No newline at end of file
diff --git a/Engine/Plugins/Experimental/AnimNext/Source/AnimNext/Private/Module/RigVMTrait_ModuleEventDependency_ActorComponentPrimaryTickFunction.cpp b/Engine/Plugins/Experimental/AnimNext/Source/AnimNext/Private/Module/RigVMTrait_ModuleEventDependency_ActorComponentPrimaryTickFunction.cpp
new file mode 100644
--- /dev/null
+++ b/Engine/Plugins/Experimental/AnimNext/Source/AnimNext/Private/Module/RigVMTrait_ModuleEventDependency_ActorComponentPrimaryTickFunction.cpp
@@ -0,0 +1,56 @@
+// Copyright Epic Games, Inc. All Rights Reserved.
+
+#include "RigVMTrait_ModuleEventDependency_ActorComponentPrimaryTickFunction.h"
+
+#if WITH_EDITOR
+FString FRigVMTrait_ModuleEventDependency_ActorComponentPrimaryTickFunction::GetDisplayName() const
+{
+	return StaticStruct()->GetDisplayNameText().ToString();
+}
+#endif
+
+void FRigVMTrait_ModuleEventDependency_ActorComponentPrimaryTickFunction::OnAddDependency(const UE::AnimNext::FModuleDependencyContext& InContext) const
+{
+	UActorComponent* AnimNextComponent = Cast<UActorComponent>(InContext.Object);
+	if(AnimNextComponent == nullptr)
+	{
+		return;
+	}
+
+	if(Component == nullptr)
+	{
+		return;
+	}
+
+	if(Ordering == EAnimNextModuleEventDependencyOrdering::Before)
+	{
+		InContext.TickFunction.AddPrerequisite(Component, Component->PrimaryComponentTick);
+	}
+	else
+	{
+		Component->PrimaryComponentTick.AddPrerequisite(AnimNextComponent, InContext.TickFunction);
+	}
+}
+
+void FRigVMTrait_ModuleEventDependency_ActorComponentPrimaryTickFunction::OnRemoveDependency(const UE::AnimNext::FModuleDependencyContext& InContext) const
+{
+	UActorComponent* AnimNextComponent = Cast<UActorComponent>(InContext.Object);
+	if(AnimNextComponent == nullptr)
+	{
+		return;
+	}
+	
+	if(Component == nullptr)
+	{
+		return;
+	}
+
+	if(Ordering == EAnimNextModuleEventDependencyOrdering::Before)
+	{
+		InContext.TickFunction.RemovePrerequisite(Component, Component->PrimaryComponentTick);
+	}
+	else
+	{
+		Component->PrimaryComponentTick.RemovePrerequisite(AnimNextComponent, InContext.TickFunction);
+	}
+}
\ No newline at end of file
diff --git a/Engine/Plugins/Experimental/AnimNext/Source/AnimNext/Private/Module/RigVMTrait_ModuleEventDependency_ActorComponentPrimaryTickFunction.h b/Engine/Plugins/Experimental/AnimNext/Source/AnimNext/Private/Module/RigVMTrait_ModuleEventDependency_ActorComponentPrimaryTickFunction.h
new file mode 100644
--- /dev/null
+++ b/Engine/Plugins/Experimental/AnimNext/Source/AnimNext/Private/Module/RigVMTrait_ModuleEventDependency_ActorComponentPrimaryTickFunction.h
@@ -0,0 +1,28 @@
+// Copyright Epic Games, Inc. All Rights Reserved.
+
+#pragma once
+
+#include "Module/RigVMTrait_ModuleEventDependency.h"
+#include "RigVMTrait_ModuleEventDependency_ActorComponentPrimaryTickFunction.generated.h"
+
+class UActorComponent;
+
+// A dependency on an actor component's primary tick function
+USTRUCT(DisplayName="Actor Component", meta=(ShowTooltip=true))
+struct FRigVMTrait_ModuleEventDependency_ActorComponentPrimaryTickFunction : public FRigVMTrait_ModuleEventDependency
+{
+	GENERATED_BODY()
+
+	// FRigVMTrait interface
+#if WITH_EDITOR
+	virtual FString GetDisplayName() const override;
+#endif
+
+	// FAnimNextModuleEventDependency interface
+	virtual void OnAddDependency(const UE::AnimNext::FModuleDependencyContext& InContext) const override;
+	virtual void OnRemoveDependency(const UE::AnimNext::FModuleDependencyContext& InContext) const override;
+
+	// The component on whose primary tick function we will depend
+	UPROPERTY(EditAnywhere, Category = "Dependency")
+	TObjectPtr<UActorComponent> Component;
+};
\ No newline at end of file
diff --git a/Engine/Plugins/Experimental/AnimNext/Source/AnimNext/Public/Module/RigVMTrait_ModuleEventDependency.h b/Engine/Plugins/Experimental/AnimNext/Source/AnimNext/Public/Module/RigVMTrait_ModuleEventDependency.h
new file mode 100644
--- /dev/null
+++ b/Engine/Plugins/Experimental/AnimNext/Source/AnimNext/Public/Module/RigVMTrait_ModuleEventDependency.h
@@ -0,0 +1,60 @@
+// Copyright Epic Games, Inc. All Rights Reserved.
+
+#pragma once
+
+#include "CoreMinimal.h"
+#include "Module/RigUnit_AnimNextModuleEvents.h"
+#include "RigVMCore/RigVMTrait.h"
+#include "RigVMTrait_ModuleEventDependency.generated.h"
+
+struct FTickFunction;
+
+// The relative ordering of a module dependency
+UENUM()
+enum class EAnimNextModuleEventDependencyOrdering : uint8
+{
+	// This dependency executes before the specified module event
+	Before,
+	
+	// This dependency executes after the specified module event
+	After
+};
+
+namespace UE::AnimNext
+{
+	// Context passed to module dependency functions
+	struct FModuleDependencyContext
+	{
+		FModuleDependencyContext(UObject* InObject, FTickFunction& InTickFunction)
+			: Object(InObject)
+			, TickFunction(InTickFunction)
+		{}
+		
+		// The object that the module is bound to
+		UObject* Object = nullptr;
+
+		// The module's tick function that we want to depend on
+		FTickFunction& TickFunction;
+	};
+}
+
+// A trait that acts as dependency that can be established between an external system and an AnimNext event
+USTRUCT(meta=(Hidden))
+struct ANIMNEXT_API FRigVMTrait_ModuleEventDependency : public FRigVMTrait
+{
+	GENERATED_BODY()
+
+	// Override point that adds the dependency
+	virtual void OnAddDependency(const UE::AnimNext::FModuleDependencyContext& InContext) const {}
+
+	// Override point that removes the dependency
+	virtual void OnRemoveDependency(const UE::AnimNext::FModuleDependencyContext& InContext) const {}
+
+	// The event to execute relative to
+	UPROPERTY(EditAnywhere, Category="Dependency", meta=(GetOptions="/Script/AnimNext.AnimNextConfig:GetModuleEventNames"))
+	FName EventName = FRigUnit_AnimNextPrePhysicsEvent::EventName;
+
+	// How to execute relative to the event
+	UPROPERTY(EditAnywhere, Category="Dependency")
+	EAnimNextModuleEventDependencyOrdering Ordering = EAnimNextModuleEventDependencyOrdering::Before;
+};
\ No newline at end of file
