New AnimNext basic Manifest
- We can now right click on Trait Stacks to include / remove them from the manifest, which adds / removes the nodes from the graph context menu (so we can paste copies of the node)
- Created a new AnimNextTraitStackRigUnitNode to enable some custom data (like Manifest information) to be added to the node.
- Added a new delegate at Module to enable external modules to add context menu options in the graph
- Added a simple test for the Manifest on AnimNextTraitStackRigUnitNode

#rnx

diff --git a/Engine/Plugins/Experimental/AnimNext/Source/AnimNextEditor/Internal/Common/GraphEditorSchemaActions.h b/Engine/Plugins/Experimental/AnimNext/Source/AnimNextEditor/Internal/Common/GraphEditorSchemaActions.h
new file mode 100644
--- /dev/null
+++ b/Engine/Plugins/Experimental/AnimNext/Source/AnimNextEditor/Internal/Common/GraphEditorSchemaActions.h
@@ -0,0 +1,156 @@
+﻿// Copyright Epic Games, Inc. All Rights Reserved.
+
+#pragma once
+
+#include "CoreMinimal.h"
+#include "EdGraph/EdGraphSchema.h"
+#include "RigVMCore/RigVMExternalVariable.h"
+#include "RigVMCore/RigVMGraphFunctionDefinition.h"
+#include "Styling/AppStyle.h"
+#include "GraphEditorSchemaActions.generated.h"
+
+struct FAnimNextParamType;
+struct FSlateBrush;
+class URigVMLibraryNode;
+class URigVMUnitNode;
+
+USTRUCT()
+struct ANIMNEXTEDITOR_API FAnimNextSchemaAction : public FEdGraphSchemaAction
+{
+	GENERATED_BODY()
+
+	FAnimNextSchemaAction() = default;
+	
+	FAnimNextSchemaAction(FText InNodeCategory, FText InMenuDesc, FText InToolTip, FText InKeywords = FText::GetEmpty())
+		: FEdGraphSchemaAction(MoveTemp(InNodeCategory), MoveTemp(InMenuDesc), MoveTemp(InToolTip), 0, MoveTemp(InKeywords))
+	{
+	}
+
+	virtual const FSlateBrush* GetIconBrush() const;
+
+	virtual const FLinearColor& GetIconColor() const;
+};
+
+USTRUCT()
+struct ANIMNEXTEDITOR_API FAnimNextSchemaAction_RigUnit : public FAnimNextSchemaAction
+{
+	GENERATED_BODY()
+
+	FAnimNextSchemaAction_RigUnit() = default;
+	
+	FAnimNextSchemaAction_RigUnit(const TSubclassOf<URigVMUnitNode>& InNodeClass, UScriptStruct* InStructTemplate, FText InNodeCategory, FText InMenuDesc, FText InToolTip, FText InKeywords = FText::GetEmpty())
+		: FAnimNextSchemaAction(InNodeCategory, InMenuDesc, InToolTip, InKeywords)
+		, StructTemplate(InStructTemplate)
+		, NodeClass(InNodeClass)
+	{}
+
+	// FEdGraphSchemaAction Interface
+	virtual UEdGraphNode* PerformAction(UEdGraph* ParentGraph, TArray<UEdGraphPin*>& FromPins, const FVector2D Location, bool bSelectNewNode = true) override;
+	virtual UEdGraphNode* PerformAction(UEdGraph* ParentGraph, UEdGraphPin* FromPin, const FVector2D Location, bool bSelectNewNode = true) { return nullptr; }
+
+private:
+	// The script struct for our rig unit
+	UScriptStruct* StructTemplate = nullptr;
+	TSubclassOf<URigVMUnitNode> NodeClass;
+};
+
+
+USTRUCT()
+struct FAnimNextSchemaAction_DispatchFactory : public FAnimNextSchemaAction
+{
+	GENERATED_BODY()
+
+	FAnimNextSchemaAction_DispatchFactory() = default;
+
+	FAnimNextSchemaAction_DispatchFactory(FName InNotation, FText InNodeCategory, FText InMenuDesc, FText InToolTip, FText InKeywords = FText::GetEmpty())
+		: FAnimNextSchemaAction(InNodeCategory, InMenuDesc, InToolTip, InKeywords)
+		, Notation(InNotation)
+	{}
+
+	virtual const FSlateBrush* GetIconBrush() const override;
+	
+	// FEdGraphSchemaAction Interface
+	virtual UEdGraphNode* PerformAction(UEdGraph* ParentGraph, TArray<UEdGraphPin*>& FromPins, const FVector2D Location, bool bSelectNewNode = true) override;
+	virtual UEdGraphNode* PerformAction(UEdGraph* ParentGraph, UEdGraphPin* FromPin, const FVector2D Location, bool bSelectNewNode = true) { return nullptr; }
+
+private:
+	// Notation for dispatch factory
+	FName Notation;
+};
+
+
+USTRUCT()
+struct FAnimNextSchemaAction_Variable : public FAnimNextSchemaAction
+{
+	GENERATED_BODY()
+
+	enum class EVariableAccessorChoice
+	{
+		Set,
+		Get,
+		Deferred
+	};
+
+	FAnimNextSchemaAction_Variable() = default;
+
+	FAnimNextSchemaAction_Variable(FName InName, const FAnimNextParamType& InType, const EVariableAccessorChoice InVariableAccessorChoice);
+
+	// FEdGraphSchemaAction Interface
+	virtual UEdGraphNode* PerformAction(UEdGraph* ParentGraph, TArray<UEdGraphPin*>& FromPins, const FVector2D Location, bool bSelectNewNode = true) override;
+	virtual UEdGraphNode* PerformAction(UEdGraph* ParentGraph, UEdGraphPin* FromPin, const FVector2D Location, bool bSelectNewNode = true) { return nullptr; }
+
+	virtual const FSlateBrush* GetIconBrush() const override;
+
+	virtual const FLinearColor& GetIconColor() const override;
+
+private:
+	FName Name;
+	FString ObjectPath;
+	FString TypeName;
+	EVariableAccessorChoice VariableAccessorChoice = EVariableAccessorChoice::Set;
+	FLinearColor VariableColor;
+};
+
+USTRUCT()
+struct FAnimNextSchemaAction_AddComment : public FAnimNextSchemaAction
+{
+	GENERATED_BODY()
+
+	FAnimNextSchemaAction_AddComment();
+
+	// FEdGraphSchemaAction Interface
+	virtual UEdGraphNode* PerformAction(UEdGraph* ParentGraph, TArray<UEdGraphPin*>& FromPins, const FVector2D Location, bool bSelectNewNode = true) override;
+	virtual UEdGraphNode* PerformAction(UEdGraph* ParentGraph, UEdGraphPin* FromPin, const FVector2D Location, bool bSelectNewNode = true) { return nullptr; }
+
+	virtual const FSlateBrush* GetIconBrush() const override;
+};
+
+USTRUCT()
+struct FAnimNextSchemaAction_Function : public FAnimNextSchemaAction
+{
+	GENERATED_BODY()
+
+	FAnimNextSchemaAction_Function() = default;
+
+	FAnimNextSchemaAction_Function(const FRigVMGraphFunctionHeader& InReferencedPublicFunctionHeader, const FText& InNodeCategory, const FText& InMenuDesc, const FText& InToolTip, const FText& InKeywords = FText::GetEmpty());
+	FAnimNextSchemaAction_Function(const URigVMLibraryNode* InFunctionLibraryNode, const FText& InNodeCategory, const FText& InMenuDesc, const FText& InToolTip, const FText& InKeywords = FText::GetEmpty());
+
+	virtual const FSlateBrush* GetIconBrush() const override;
+	
+	// FEdGraphSchemaAction Interface
+	virtual UEdGraphNode* PerformAction(UEdGraph* ParentGraph, TArray<UEdGraphPin*>& FromPins, const FVector2D Location, bool bSelectNewNode = true) override;
+	virtual UEdGraphNode* PerformAction(UEdGraph* ParentGraph, UEdGraphPin* FromPin, const FVector2D Location, bool bSelectNewNode = true) { return nullptr; }
+
+private:
+	/** The public function definition we will spawn from [optional] */
+	UPROPERTY(Transient)
+	FRigVMGraphFunctionHeader ReferencedPublicFunctionHeader;
+
+	/** Marked as true for local function definitions */
+	UPROPERTY(Transient)
+	bool bIsLocalFunction = false;
+
+	/** Holds the node type that this spawner will instantiate. */
+	UPROPERTY(Transient)
+	TSubclassOf<UEdGraphNode> NodeClass;
+};
\ No newline at end of file
diff --git a/Engine/Plugins/Experimental/AnimNext/Source/AnimNextEditor/Internal/Common/SActionMenu.h b/Engine/Plugins/Experimental/AnimNext/Source/AnimNextEditor/Internal/Common/SActionMenu.h
new file mode 100644
--- /dev/null
+++ b/Engine/Plugins/Experimental/AnimNext/Source/AnimNextEditor/Internal/Common/SActionMenu.h
@@ -0,0 +1,96 @@
+﻿// Copyright Epic Games, Inc. All Rights Reserved.
+#pragma once
+
+#include "GraphEditor.h"
+#include "EdGraph/EdGraph.h"
+#include "EdGraph/EdGraphPin.h"
+#include "EdGraph/EdGraphSchema.h"
+#include "Layout/Visibility.h"
+#include "SGraphActionMenu.h"
+#include "Widgets/DeclarativeSyntaxSupport.h"
+#include "Widgets/Input/SEditableTextBox.h"
+#include "Widgets/Layout/SBorder.h"
+#include "Widgets/SWidget.h"
+
+class UAnimNextRigVMAssetEditorData;
+class URigVMSchema;
+class SEditableTextBox;
+class SGraphActionMenu;
+class UEdGraph;
+class IRigVMClientHost;
+class URigVMController;
+class URigVMHost;
+
+namespace UE::AnimNext::Editor
+{
+
+struct FActionMenuContextData
+{
+	TArray<UObject*> SelectedObjects;
+
+	UEdGraph* Graph = nullptr;
+	const URigVMSchema* RigVMSchema = nullptr;
+	URigVMHost* RigVMHost = nullptr;
+	IRigVMClientHost* RigVMClientHost = nullptr;
+	URigVMController* RigVMController = nullptr;
+	UAnimNextRigVMAssetEditorData* EditorData = nullptr;
+	bool bShowGlobalManifestNodes = true;
+};
+
+class SActionMenu : public SBorder
+{
+public:
+	/** Delegate to retrieve the action list for the graph */
+	DECLARE_DELEGATE_TwoParams(FCollectAllGraphActions, FGraphContextMenuBuilder& /*MenuBuilder*/, const FActionMenuContextData& /*ActionMenuContextData*/);
+
+	/** Delegate for the OnCloseReason event which is always raised when the SActionMenu closes */
+	DECLARE_DELEGATE_ThreeParams(FClosedReason, bool /*bActionExecuted*/, bool /*bContextSensitiveChecked*/, bool /*bGraphPinContext*/);
+
+	SLATE_BEGIN_ARGS(SActionMenu)
+		: _NewNodePosition(FVector2D::ZeroVector)
+		, _AutoExpandActionMenu(false)
+	{}
+	
+	SLATE_ARGUMENT(FVector2D, NewNodePosition)
+	SLATE_ARGUMENT(TArray<UEdGraphPin*>, DraggedFromPins)
+	SLATE_ARGUMENT(SGraphEditor::FActionMenuClosed, OnClosedCallback)
+	SLATE_ARGUMENT(bool, AutoExpandActionMenu)
+	SLATE_EVENT(FClosedReason, OnCloseReason)
+	SLATE_EVENT(FCollectAllGraphActions, OnCollectGraphActionsCallback)
+
+	SLATE_END_ARGS()
+
+	void Construct(const FArguments& InArgs, UEdGraph* InGraph);
+
+	~SActionMenu();
+
+	TSharedRef<SEditableTextBox> GetFilterTextBox();
+
+protected:
+	void OnActionSelected(const TArray<TSharedPtr<FEdGraphSchemaAction>>& SelectedAction, ESelectInfo::Type InSelectionType);
+
+	TSharedRef<SWidget> OnCreateWidgetForAction(struct FCreateWidgetForActionData* const InCreateData);
+
+	/** Callback used to populate all actions list in SGraphActionMenu */
+	void CollectAllActions(FGraphActionListBuilderBase& OutAllActions);
+
+	void CollectAllAnimNextGraphActions(FGraphContextMenuBuilder& MenuBuilder) const;
+	
+private:
+	FActionMenuContextData ContextData;
+
+	bool bAutoExpandActionMenu = false;
+	bool bActionExecuted = false;
+
+	TArray<UEdGraphPin*> DraggedFromPins;
+	FVector2D NewNodePosition = FVector2D::ZeroVector;
+
+	SGraphEditor::FActionMenuClosed OnClosedCallback;
+	FClosedReason OnCloseReasonCallback;
+
+	FCollectAllGraphActions OnCollectGraphActionsCallback;
+
+	TSharedPtr<SGraphActionMenu> GraphActionMenu;
+};
+
+}
diff --git a/Engine/Plugins/Experimental/AnimNext/Source/AnimNextEditor/Internal/EditorUtils.h b/Engine/Plugins/Experimental/AnimNext/Source/AnimNextEditor/Internal/EditorUtils.h
new file mode 100644
--- /dev/null
+++ b/Engine/Plugins/Experimental/AnimNext/Source/AnimNextEditor/Internal/EditorUtils.h
@@ -0,0 +1,45 @@
+// Copyright Epic Games, Inc. All Rights Reserved.
+
+#pragma once
+
+#include "CoreMinimal.h"
+#include "Param/ParamType.h"
+#include "EdGraphSchema_K2.h"
+#include "StructView.h"
+
+struct FAnimNextVariableBindingData;
+class UAnimNextModule;
+class UAnimNextRigVMAssetEditorData;
+struct FAnimNextParamType;
+class URigVMController;
+struct FAnimNextWorkspaceAssetRegistryExports;
+class SWidget;
+class URigVMUnitNode;
+struct FRigVMFunction;
+struct FGraphContextMenuBuilder;
+
+struct FAnimNextAssetRegistryExports;
+
+namespace UE::AnimNext::Editor
+{
+
+struct FUtils
+{
+	static ANIMNEXTEDITOR_API FName ValidateName(const UObject* InObject, const FString& InName);
+
+	static ANIMNEXTEDITOR_API void GetAllEntryNames(const UAnimNextRigVMAssetEditorData* InEditorData, TSet<FName>& OutNames);
+
+	static ANIMNEXTEDITOR_API FAnimNextParamType GetParameterTypeFromMetaData(const FStringView& InStringView);
+
+	static ANIMNEXTEDITOR_API void GetFilteredVariableTypeTree(TArray<TSharedPtr<UEdGraphSchema_K2::FPinTypeTreeInfo>>& TypeTree, ETypeTreeFilter TypeTreeFilter);
+
+	static ANIMNEXTEDITOR_API bool IsValidParameterNameString(FStringView InStringView, FText& OutErrorText);
+
+	static ANIMNEXTEDITOR_API bool IsValidParameterName(const FName InName, FText& OutErrorText);
+
+	static ANIMNEXTEDITOR_API bool AddSchemaRigUnitAction(const TSubclassOf<URigVMUnitNode>& UnitNodeClass, UScriptStruct* Struct, const FRigVMFunction& Function, FGraphContextMenuBuilder& InContexMenuBuilder);
+
+	static ANIMNEXTEDITOR_API void GetRigUnitStructMetadata(const UScriptStruct* Struct, FString& OutCategoryMetadata, FString& OutDisplayNameMetadata, FString& OutMenuDescSuffixMetadata);
+};
+
+}
diff --git a/Engine/Plugins/Experimental/AnimNext/Source/AnimNextEditor/Private/AnimNextEditorModule.cpp b/Engine/Plugins/Experimental/AnimNext/Source/AnimNextEditor/Private/AnimNextEditorModule.cpp
--- a/Engine/Plugins/Experimental/AnimNext/Source/AnimNextEditor/Private/AnimNextEditorModule.cpp
+++ b/Engine/Plugins/Experimental/AnimNext/Source/AnimNextEditor/Private/AnimNextEditorModule.cpp
@@ -51,6 +51,9 @@
 #include "Variables/VariableOverrideCommands.h"
 #include "Variables/VariableProxyCustomization.h"
 #include "Workspace/AnimNextWorkspaceSchema.h"
+#include "ExternalPackageHelper.h"
+#include "Common/GraphEditorSchemaActions.h"
+#include "EditorUtils.h"
 
 #define LOCTEXT_NAMESPACE "AnimNextEditorModule"
 
@@ -257,6 +260,158 @@ void FAnimNextEditorModule::RemoveWorkspaceSupportedAssetClass(const FTopLevelAs
 	}	
 }
 
+FDelegateHandle FAnimNextEditorModule::RegisterGraphMenuActionsProvider(const FOnCollectGraphMenuActionsDelegate& CollectDelegate)
+{
+	return OnCollectGraphMenuActionsDelegateImpl.Add(CollectDelegate);
+}
+
+void FAnimNextEditorModule::UnregisterGraphMenuActionsProvider(const FDelegateHandle& DelegateHandle)
+{
+	OnCollectGraphMenuActionsDelegateImpl.Remove(DelegateHandle);
+}
+
+void FAnimNextEditorModule::CollectGraphMenuActions(const TWeakPtr<UE::Workspace::IWorkspaceEditor>& WorkspaceEditor, FGraphContextMenuBuilder& InContexMenuBuilder, const FActionMenuContextData& InActionMenuContextData)
+{
+	if (OnCollectGraphMenuActionsDelegateImpl.IsBound())
+	{
+		OnCollectGraphMenuActionsDelegateImpl.Broadcast(WorkspaceEditor, InContexMenuBuilder, InActionMenuContextData);
+	}
+
+	if (const URigVMEdGraph* RigVMEdGraph = Cast<URigVMEdGraph>(InContexMenuBuilder.CurrentGraph))
+	{
+		const URigVMGraph* Graph = RigVMEdGraph->GetModel();
+		check(Graph);
+
+		const IRigVMClientHost* RigVMClientHost = InActionMenuContextData.RigVMClientHost;
+		check(RigVMClientHost);
+		const URigVMHost* RigVMHost = InActionMenuContextData.RigVMHost;
+		check(RigVMHost);
+		URigVMController* RigVMController = InActionMenuContextData.RigVMController;
+		check(RigVMController);
+		const URigVMSchema* RigVMSchema = InActionMenuContextData.RigVMSchema;
+		check(RigVMSchema);
+		const UAnimNextRigVMAssetEditorData* EditorData = InActionMenuContextData.EditorData;
+		check(EditorData);
+
+		for (const FRigVMFunction& Function : FRigVMRegistry::Get().GetFunctions())
+		{
+			if (RigVMSchema == nullptr || !RigVMSchema->SupportsUnitFunction(RigVMController, &Function))
+			{
+				continue;
+			}
+
+			UScriptStruct* Struct = Function.Struct;
+			if (Struct == nullptr)
+			{
+				continue;
+			}
+
+			// skip deprecated units
+			if (Function.Struct->HasMetaData(FRigVMStruct::DeprecatedMetaName))
+			{
+				continue;
+			}
+
+			// skip hidden units
+			if (Function.Struct->HasMetaData(FRigVMStruct::HiddenMetaName))
+			{
+				continue;
+			}
+
+			// Disallow trait stacks to be added here, as it will be added at AnimNextAnimGraph with a custom node class
+			static const UScriptStruct* TraitStackStruct = FindObjectChecked<UScriptStruct>(nullptr, TEXT("/Script/AnimNextAnimGraph.RigUnit_AnimNextTraitStack"));
+			if (Struct && Struct->IsChildOf(TraitStackStruct))
+			{
+				continue;
+			}
+
+			UE::AnimNext::Editor::FUtils::AddSchemaRigUnitAction(URigVMUnitNode::StaticClass(), Struct, Function, InContexMenuBuilder);
+		}
+
+		for (const FRigVMDispatchFactory* Factory : FRigVMRegistry::Get().GetFactories())
+		{
+			if (RigVMSchema == nullptr || !RigVMSchema->SupportsDispatchFactory(RigVMController, Factory))
+			{
+				continue;
+			}
+
+			const FRigVMTemplate* Template = Factory->GetTemplate();
+			if (Template == nullptr)
+			{
+				continue;
+			}
+
+			// skip deprecated factories
+			if (Factory->GetScriptStruct()->HasMetaData(FRigVMStruct::DeprecatedMetaName))
+			{
+				continue;
+			}
+
+			// skip hidden factories
+			if (Factory->GetScriptStruct()->HasMetaData(FRigVMStruct::HiddenMetaName))
+			{
+				continue;
+			}
+
+			FText NodeCategory = FText::FromString(Factory->GetCategory());
+			FText MenuDesc = FText::FromString(Factory->GetNodeTitle(FRigVMTemplateTypeMap()));
+			FText ToolTip = Factory->GetNodeTooltip(FRigVMTemplateTypeMap());
+
+			InContexMenuBuilder.AddAction(MakeShared<FAnimNextSchemaAction_DispatchFactory>(Template->GetNotation(), NodeCategory, MenuDesc, ToolTip));
+		};
+
+		if (URigVMFunctionLibrary* LocalFunctionLibrary = RigVMClientHost->GetLocalFunctionLibrary())
+		{
+			const FSoftObjectPath LocalLibrarySoftPath = LocalFunctionLibrary->GetFunctionHostObjectPath();
+
+			TArray<URigVMLibraryNode*> Functions = LocalFunctionLibrary->GetFunctions();
+			for (URigVMLibraryNode* FunctionLibraryNode : Functions)
+			{
+				if (LocalFunctionLibrary->IsFunctionPublic(FunctionLibraryNode->GetFName()))	// Public functions will be added when processing asset registry exports
+				{
+					continue;
+				}
+				const FText NodeCategory = FText::FromString(FunctionLibraryNode->GetNodeCategory());
+				const FText MenuDesc = FText::FromString(FunctionLibraryNode->GetName());
+				const FText ToolTip = FunctionLibraryNode->GetToolTipText();
+
+				InContexMenuBuilder.AddAction(MakeShared<FAnimNextSchemaAction_Function>(FunctionLibraryNode, NodeCategory, MenuDesc, ToolTip));
+			}
+		}
+
+		TMap<FAssetData, FRigVMGraphFunctionHeaderArray> FunctionExports;
+		AnimNext::UncookedOnly::FUtils::GetExportedFunctionsFromAssetRegistry(UE::AnimNext::AnimNextPublicGraphFunctionsExportsRegistryTag, FunctionExports);
+		// TODO: Ideally we can filter functions by schema or execute context, but right now we dont expose the schema and function execute contexts are
+		// all FRigVMExecuteContext, rather than the 'most derived' context in the function.
+		//	AnimNext::UncookedOnly::FUtils::GetExportedFunctionsFromAssetRegistry(UE::AnimNext::ControlRigAssetPublicGraphFunctionsExportsRegistryTag, FunctionExports);
+
+		for (const auto& Export : FunctionExports.Array())
+		{
+			for (const FRigVMGraphFunctionHeader& FunctionHeader : Export.Value.Headers)
+			{
+				if (FunctionHeader.LibraryPointer.IsValid())
+				{
+					const FText NodeCategory = FText::FromString(FunctionHeader.Category);
+					const FText MenuDesc = FText::FromString(FunctionHeader.NodeTitle);
+					const FText ToolTip = FunctionHeader.GetTooltip();
+
+					InContexMenuBuilder.AddAction(MakeShared<FAnimNextSchemaAction_Function>(FunctionHeader, NodeCategory, MenuDesc, ToolTip));
+				}
+			}
+		}
+
+		TArray<UAnimNextVariableEntry*> Variables;
+		EditorData->GetAllVariables(Variables);
+		for (UAnimNextVariableEntry* VariableEntry : Variables)
+		{
+			InContexMenuBuilder.AddAction(MakeShared<FAnimNextSchemaAction_Variable>(VariableEntry->GetVariableName(), VariableEntry->GetType(), FAnimNextSchemaAction_Variable::EVariableAccessorChoice::Set));
+			InContexMenuBuilder.AddAction(MakeShared<FAnimNextSchemaAction_Variable>(VariableEntry->GetVariableName(), VariableEntry->GetType(), FAnimNextSchemaAction_Variable::EVariableAccessorChoice::Get));
+		}
+
+		InContexMenuBuilder.AddAction(MakeShared<FAnimNextSchemaAction_AddComment>());
+	}
+}
+
 void FAnimNextEditorModule::RegisterWorkspaceDocumentTypes(Workspace::IWorkspaceEditorModule& WorkspaceEditorModule)
 {
 	// --- AnimNextModule ---
@@ -280,13 +435,20 @@ void FAnimNextEditorModule::RegisterWorkspaceDocumentTypes(Workspace::IWorkspace
 	// --- AnimNextEdGraph ---
 	Workspace::FGraphDocumentWidgetArgs GraphArgs;
 	GraphArgs.SpawnLocation = Workspace::WorkspaceTabs::TopMiddleDocumentArea;
-	GraphArgs.OnCreateActionMenu = Workspace::FOnCreateActionMenu::CreateLambda([](const Workspace::FWorkspaceEditorContext& InContext, UEdGraph* InGraph, const FVector2D& InNodePosition, const TArray<UEdGraphPin*>& InDraggedPins, bool bAutoExpand, SGraphEditor::FActionMenuClosed InOnMenuClosed)
+	GraphArgs.OnCreateActionMenu = Workspace::FOnCreateActionMenu::CreateLambda([this](const Workspace::FWorkspaceEditorContext& InContext, UEdGraph* InGraph, const FVector2D& InNodePosition, const TArray<UEdGraphPin*>& InDraggedPins, bool bAutoExpand, SGraphEditor::FActionMenuClosed InOnMenuClosed)
 	{
+		const FString WorkspaceAssetPath = FExternalPackageHelper::GetExternalObjectsPath(InContext.WorkspaceEditor->GetPackageName());
+		TWeakPtr<UE::Workspace::IWorkspaceEditor> WorkspaceEditorWeak = InContext.WorkspaceEditor.ToWeakPtr();
+
 		TSharedRef<SActionMenu> ActionMenu = SNew(SActionMenu, InGraph)
 			.AutoExpandActionMenu(bAutoExpand)
 			.NewNodePosition(InNodePosition)
 			.DraggedFromPins(InDraggedPins)
-			.OnClosedCallback(InOnMenuClosed);
+			.OnClosedCallback(InOnMenuClosed)
+			.OnCollectGraphActionsCallback(SActionMenu::FCollectAllGraphActions::CreateLambda([this, WorkspaceEditorWeak](FGraphContextMenuBuilder& InContexMenuBuilder, const FActionMenuContextData& InActionMenuContextData)
+				{
+					CollectGraphMenuActions(WorkspaceEditorWeak, InContexMenuBuilder, InActionMenuContextData);
+				}));
 
 		TSharedPtr<SWidget> FilterTextBox = StaticCastSharedRef<SWidget>(ActionMenu->GetFilterTextBox());
 		return FActionMenuContent(StaticCastSharedRef<SWidget>(ActionMenu), FilterTextBox);
diff --git a/Engine/Plugins/Experimental/AnimNext/Source/AnimNextEditor/Private/AnimNextEditorModule.h b/Engine/Plugins/Experimental/AnimNext/Source/AnimNextEditor/Private/AnimNextEditorModule.h
--- a/Engine/Plugins/Experimental/AnimNext/Source/AnimNextEditor/Private/AnimNextEditorModule.h
+++ b/Engine/Plugins/Experimental/AnimNext/Source/AnimNextEditor/Private/AnimNextEditorModule.h
@@ -1,4 +1,4 @@
-﻿// Copyright Epic Games, Inc. All Rights Reserved.
+// Copyright Epic Games, Inc. All Rights Reserved.
 
 #pragma once
 
@@ -34,7 +34,12 @@ private:
 	virtual void UnregisterLocatorFragmentEditorType(FName InLocatorFragmentEditorName) override;
 	virtual void AddWorkspaceSupportedAssetClass(const FTopLevelAssetPath& InClassAssetPath) override;
 	virtual void RemoveWorkspaceSupportedAssetClass(const FTopLevelAssetPath& InClassAssetPath) override;
+	virtual FDelegateHandle RegisterGraphMenuActionsProvider(const FOnCollectGraphMenuActionsDelegate& InCollectDelegate) override;
+	virtual void UnregisterGraphMenuActionsProvider(const FDelegateHandle& InDelegateHandle) override;
+
 private:
+	void CollectGraphMenuActions(const TWeakPtr<UE::Workspace::IWorkspaceEditor>& WorkspaceEditor, FGraphContextMenuBuilder& InContexMenuBuilder, const FActionMenuContextData& InActionMenuContextData);
+
 	void RegisterWorkspaceDocumentTypes(Workspace::IWorkspaceEditorModule& WorkspaceEditorModule);
 	void UnregisterWorkspaceDocumentTypes();
 
@@ -47,6 +52,8 @@ private:
 
 	TArray<FTopLevelAssetPath> SupportedAssetClasses;
 
+	FCollectGraphMenuActionsMulticast OnCollectGraphMenuActionsDelegateImpl;
+
 	friend class FLocatorContext;
 };
 
diff --git a/Engine/Plugins/Experimental/AnimNext/Source/AnimNextEditor/Private/Common/GraphEditorSchemaActions.cpp b/Engine/Plugins/Experimental/AnimNext/Source/AnimNextEditor/Private/Common/GraphEditorSchemaActions.cpp
--- a/Engine/Plugins/Experimental/AnimNext/Source/AnimNextEditor/Private/Common/GraphEditorSchemaActions.cpp
+++ b/Engine/Plugins/Experimental/AnimNext/Source/AnimNextEditor/Private/Common/GraphEditorSchemaActions.cpp
@@ -1,6 +1,6 @@
 ﻿// Copyright Epic Games, Inc. All Rights Reserved.
 
-#include "GraphEditorSchemaActions.h"
+#include "Common/GraphEditorSchemaActions.h"
 #include "EditorUtils.h"
 #include "AnimNextEdGraph.h"
 #include "AnimNextEdGraphNode.h"
@@ -14,9 +14,26 @@
 #include "GraphEditor.h"
 #include "Framework/MultiBox/MultiBoxBuilder.h"
 #include "Framework/Application/SlateApplication.h"
+#include "AnimNextUnitNode.h"
+#include "Editor/RigVMEditorStyle.h"
+
 
 #define LOCTEXT_NAMESPACE "AnimNextSchemaActions"
 
+// *** Base Schema Action ***
+
+const FSlateBrush* FAnimNextSchemaAction::GetIconBrush() const
+{
+	return FRigVMEditorStyle::Get().GetBrush("RigVM.Unit");
+}
+
+const FLinearColor& FAnimNextSchemaAction::GetIconColor() const
+{
+	return FLinearColor::White;
+}
+
+// *** Rig Unit ***
+
 UEdGraphNode* FAnimNextSchemaAction_RigUnit::PerformAction(UEdGraph* ParentGraph, TArray<UEdGraphPin*>& FromPins, const FVector2D Location, bool bSelectNewNode)
 {
 	IRigVMClientHost* Host = ParentGraph->GetImplementingOuter<IRigVMClientHost>();
@@ -39,14 +56,14 @@ UEdGraphNode* FAnimNextSchemaAction_RigUnit::PerformAction(UEdGraph* ParentGraph
 		FRigVMUnitNodeCreatedContext& UnitNodeCreatedContext = Controller->GetUnitNodeCreatedContext();
 		FRigVMUnitNodeCreatedContext::FScope ReasonScope(UnitNodeCreatedContext, ERigVMNodeCreatedReason::NodeSpawner);
 
-		if (URigVMUnitNode* ModelNode = Controller->AddUnitNode(StructTemplate, FRigVMStruct::ExecuteName, Location, Name.ToString(), true, false))
+		if (URigVMUnitNode* ModelNode = Controller->AddUnitNode(StructTemplate, NodeClass, FRigVMStruct::ExecuteName, Location, Name.ToString(), true, false))
 		{
 			NewNode = Cast<URigVMEdGraphNode>(EdGraph->FindNodeForModelNodeName(ModelNode->GetFName()));
 			check(NewNode);
 
 			if (NewNode)
 			{
-				if(FromPin)
+				if (FromPin)
 				{
 					NewNode->AutowireNewNode(FromPin);
 				}
@@ -71,13 +88,19 @@ UEdGraphNode* FAnimNextSchemaAction_RigUnit::PerformAction(UEdGraph* ParentGraph
 			}
 		}
 #endif
-		
+
 		Controller->CloseUndoBracket();
 	}
 
 	return NewNode;
 }
 
+// *** Dispatch Factory ***
+
+const FSlateBrush* FAnimNextSchemaAction_DispatchFactory::GetIconBrush() const
+{
+	return FRigVMEditorStyle::Get().GetBrush("RigVM.Template");
+}
 
 UEdGraphNode* FAnimNextSchemaAction_DispatchFactory::PerformAction(UEdGraph* ParentGraph, TArray<UEdGraphPin*>& FromPins, const FVector2D Location, bool bSelectNewNode)
 {
@@ -139,6 +162,8 @@ UEdGraphNode* FAnimNextSchemaAction_DispatchFactory::PerformAction(UEdGraph* Par
 	return NewNode;
 }
 
+// *** Variable ***
+
 FAnimNextSchemaAction_Variable::FAnimNextSchemaAction_Variable(FName InName, const FAnimNextParamType& InType, const EVariableAccessorChoice InVariableAccessorChoice)
 	: Name(InName)
 	, VariableAccessorChoice(InVariableAccessorChoice)
@@ -178,6 +203,16 @@ FAnimNextSchemaAction_Variable::FAnimNextSchemaAction_Variable(FName InName, con
 	VariableColor = GetDefault<UEdGraphSchema_K2>()->GetPinTypeColor(PinType);
 }
 
+const FSlateBrush* FAnimNextSchemaAction_Variable::GetIconBrush() const
+{
+	return  FAppStyle::GetBrush(TEXT("Kismet.VariableList.TypeIcon"));
+}
+
+const FLinearColor& FAnimNextSchemaAction_Variable::GetIconColor() const
+{
+	return VariableColor;
+}
+
 UEdGraphNode* FAnimNextSchemaAction_Variable::PerformAction(UEdGraph* ParentGraph, TArray<UEdGraphPin*>& FromPins, const FVector2D Location, bool bSelectNewNode)
 {
 	const auto AddNodeLambda = [](UEdGraph* ParentGraph, FName InName, const FString& InTypeName, const FString& InObjectPath, const FVector2D Location, const bool bIsGetter) -> URigVMEdGraphNode*
@@ -282,6 +317,8 @@ UEdGraphNode* FAnimNextSchemaAction_Variable::PerformAction(UEdGraph* ParentGrap
 	}
 }
 
+// *** Add Comment ***
+
 FAnimNextSchemaAction_AddComment::FAnimNextSchemaAction_AddComment()
 	: FAnimNextSchemaAction(FText(), LOCTEXT("AddComment", "Add Comment..."), LOCTEXT("AddCommentTooltip", "Create a resizable comment box."))
 {
@@ -305,6 +342,11 @@ UEdGraphNode* FAnimNextSchemaAction_AddComment::PerformAction(UEdGraph* ParentGr
 	return FEdGraphSchemaAction_NewNode::SpawnNodeFromTemplate<UEdGraphNode_Comment>(ParentGraph, CommentTemplate, SpawnLocation, bSelectNewNode);
 }
 
+const FSlateBrush* FAnimNextSchemaAction_AddComment::GetIconBrush() const
+{
+	return FAppStyle::Get().GetBrush("Icons.Comment");
+}
+
 // *** Graph Function ***
 
 FAnimNextSchemaAction_Function::FAnimNextSchemaAction_Function(const FRigVMGraphFunctionHeader& InReferencedPublicFunctionHeader, const FText& InNodeCategory, const FText& InMenuDesc, const FText& InToolTip, const FText& InKeywords)
diff --git a/Engine/Plugins/Experimental/AnimNext/Source/AnimNextEditor/Private/Common/GraphEditorSchemaActions.h b/Engine/Plugins/Experimental/AnimNext/Source/AnimNextEditor/Private/Common/GraphEditorSchemaActions.h
deleted file mode 100644
diff --git a/Engine/Plugins/Experimental/AnimNext/Source/AnimNextEditor/Private/Common/SActionMenu.cpp b/Engine/Plugins/Experimental/AnimNext/Source/AnimNextEditor/Private/Common/SActionMenu.cpp
--- a/Engine/Plugins/Experimental/AnimNext/Source/AnimNextEditor/Private/Common/SActionMenu.cpp
+++ b/Engine/Plugins/Experimental/AnimNext/Source/AnimNextEditor/Private/Common/SActionMenu.cpp
@@ -1,8 +1,8 @@
 ﻿// Copyright Epic Games, Inc. All Rights Reserved.
 
-#include "SActionMenu.h"
+#include "Common/SActionMenu.h"
 
-#include "GraphEditorSchemaActions.h"
+#include "Common/GraphEditorSchemaActions.h"
 #include "Framework/Application/SlateApplication.h"
 #include "IDocumentation.h"
 #include "SSubobjectEditor.h"
@@ -29,140 +29,14 @@ namespace UE::AnimNext::Editor
 void SActionMenu::CollectAllAnimNextGraphActions(FGraphContextMenuBuilder& MenuBuilder) const
 {
 	// Disable reporting as the schema's SupportsX functions will output errors
-	RigVMController->EnableReporting(false);
+	ContextData.RigVMController->EnableReporting(false);
 
-	for(const FRigVMFunction& Function : FRigVMRegistry::Get().GetFunctions())
+	if (OnCollectGraphActionsCallback.IsBound())
 	{
-		if (RigVMSchema == nullptr || !RigVMSchema->SupportsUnitFunction(RigVMController, &Function))
-		{
-			continue;
-		}
-
-		UScriptStruct* Struct = Function.Struct;
-		if (Struct == nullptr)
-		{
-			continue;
-		}
-
-		// skip deprecated units
-		if(Function.Struct->HasMetaData(FRigVMStruct::DeprecatedMetaName))
-		{
-			continue;
-		}
-
-		// skip hidden units
-		if(Function.Struct->HasMetaData(FRigVMStruct::HiddenMetaName))
-		{
-			continue;
-		}
-
-		FString CategoryMetadata, DisplayNameMetadata, MenuDescSuffixMetadata;
-		Struct->GetStringMetaDataHierarchical(FRigVMStruct::CategoryMetaName, &CategoryMetadata);
-		Struct->GetStringMetaDataHierarchical(FRigVMStruct::DisplayNameMetaName, &DisplayNameMetadata);
-		if(DisplayNameMetadata.IsEmpty())
-		{
-			DisplayNameMetadata = Function.GetMethodName().ToString();
-		}
-		Struct->GetStringMetaDataHierarchical(FRigVMStruct::MenuDescSuffixMetaName, &MenuDescSuffixMetadata);
-		if (!MenuDescSuffixMetadata.IsEmpty())
-		{
-			MenuDescSuffixMetadata = TEXT(" ") + MenuDescSuffixMetadata;
-		}
-		const FText NodeCategory = FText::FromString(CategoryMetadata);
-		const FText MenuDesc = FText::FromString(DisplayNameMetadata + MenuDescSuffixMetadata);
-		const FText ToolTip = Struct->GetToolTipText();
-
-		if (MenuDesc.IsEmpty())
-		{
-			continue;
-		}
-
-		MenuBuilder.AddAction(MakeShared<FAnimNextSchemaAction_RigUnit>(Struct, NodeCategory, MenuDesc, ToolTip));
+		OnCollectGraphActionsCallback.Execute(MenuBuilder, ContextData);
 	}
 
-	for (const FRigVMDispatchFactory* Factory : FRigVMRegistry::Get().GetFactories())
-	{
-		if (RigVMSchema == nullptr || !RigVMSchema->SupportsDispatchFactory(RigVMController, Factory))
-		{
-			continue;
-		}
-
-		const FRigVMTemplate* Template = Factory->GetTemplate();
-		if (Template == nullptr)
-		{
-			continue;
-		}
-
-		// skip deprecated factories
-		if(Factory->GetScriptStruct()->HasMetaData(FRigVMStruct::DeprecatedMetaName))
-		{
-			continue;
-		}
-
-		// skip hidden factories
-		if(Factory->GetScriptStruct()->HasMetaData(FRigVMStruct::HiddenMetaName))
-		{
-			continue;
-		}
-
-		FText NodeCategory = FText::FromString(Factory->GetCategory());
-		FText MenuDesc = FText::FromString(Factory->GetNodeTitle(FRigVMTemplateTypeMap()));
-		FText ToolTip = Factory->GetNodeTooltip(FRigVMTemplateTypeMap());
-
-		MenuBuilder.AddAction(MakeShared<FAnimNextSchemaAction_DispatchFactory>(Template->GetNotation(), NodeCategory, MenuDesc, ToolTip));
-	};
-
-	if (URigVMFunctionLibrary* LocalFunctionLibrary = RigVMClientHost->GetLocalFunctionLibrary())
-	{
-		const FSoftObjectPath LocalLibrarySoftPath = LocalFunctionLibrary->GetFunctionHostObjectPath();
-
-		TArray<URigVMLibraryNode*> Functions = LocalFunctionLibrary->GetFunctions();
-		for (URigVMLibraryNode* FunctionLibraryNode : Functions)
-		{
-			if (LocalFunctionLibrary->IsFunctionPublic(FunctionLibraryNode->GetFName()))	// Public functions will be added when processing asset registry exports
-			{
-				continue;
-			}
-			const FText NodeCategory = FText::FromString(FunctionLibraryNode->GetNodeCategory());
-			const FText MenuDesc = FText::FromString(FunctionLibraryNode->GetName());
-			const FText ToolTip = FunctionLibraryNode->GetToolTipText();
-
-			MenuBuilder.AddAction(MakeShared<FAnimNextSchemaAction_Function>(FunctionLibraryNode, NodeCategory, MenuDesc, ToolTip));
-		}
-	}
-
-	TMap<FAssetData, FRigVMGraphFunctionHeaderArray> FunctionExports;
-	AnimNext::UncookedOnly::FUtils::GetExportedFunctionsFromAssetRegistry(UE::AnimNext::AnimNextPublicGraphFunctionsExportsRegistryTag, FunctionExports);
-	// TODO: Ideally we can filter functions by schema or execute context, but right now we dont expose the schema and function execute contexts are
-	// all FRigVMExecuteContext, rather than the 'most derived' context in the function.
-//	AnimNext::UncookedOnly::FUtils::GetExportedFunctionsFromAssetRegistry(UE::AnimNext::ControlRigAssetPublicGraphFunctionsExportsRegistryTag, FunctionExports);
-
-	for (const auto& Export : FunctionExports.Array())
-	{
-		for (const FRigVMGraphFunctionHeader& FunctionHeader : Export.Value.Headers)
-		{
-			if (FunctionHeader.LibraryPointer.IsValid())
-			{
-				const FText NodeCategory = FText::FromString(FunctionHeader.Category);
-				const FText MenuDesc = FText::FromString(FunctionHeader.NodeTitle);
-				const FText ToolTip = FunctionHeader.GetTooltip();
-
-				MenuBuilder.AddAction(MakeShared<FAnimNextSchemaAction_Function>(FunctionHeader, NodeCategory, MenuDesc, ToolTip));
-			}
-		}
-	}
-
-	TArray<UAnimNextVariableEntry*> Variables;
-	EditorData->GetAllVariables(Variables);
-	for(UAnimNextVariableEntry* VariableEntry : Variables)
-	{
-		MenuBuilder.AddAction(MakeShared<FAnimNextSchemaAction_Variable>(VariableEntry->GetVariableName(), VariableEntry->GetType(), FAnimNextSchemaAction_Variable::EVariableAccessorChoice::Set));
-		MenuBuilder.AddAction(MakeShared<FAnimNextSchemaAction_Variable>(VariableEntry->GetVariableName(), VariableEntry->GetType(), FAnimNextSchemaAction_Variable::EVariableAccessorChoice::Get));
-	}
-
-	MenuBuilder.AddAction(MakeShared<FAnimNextSchemaAction_AddComment>());
-
-	RigVMController->EnableReporting(true);
+	ContextData.RigVMController->EnableReporting(true);
 }
 
 SActionMenu::~SActionMenu()
@@ -175,23 +49,24 @@ void SActionMenu::Construct(const FArguments& InArgs, UEdGraph* InGraph)
 {
 	check(InGraph);
 
-	Graph = InGraph;
+	ContextData.Graph = InGraph;
 	DraggedFromPins = InArgs._DraggedFromPins;
 	NewNodePosition = InArgs._NewNodePosition;
 	OnClosedCallback = InArgs._OnClosedCallback;
+	OnCollectGraphActionsCallback = InArgs._OnCollectGraphActionsCallback;
 	bAutoExpandActionMenu = InArgs._AutoExpandActionMenu;
 	OnCloseReasonCallback = InArgs._OnCloseReason;
 
-	RigVMClientHost = Graph->GetImplementingOuter<IRigVMClientHost>();
-	check(RigVMClientHost);
-	RigVMHost = Graph->GetTypedOuter<URigVMHost>();
-	check(RigVMHost);
-	RigVMController = RigVMClientHost->GetRigVMClient()->GetController(Graph);
-	check(RigVMController);
-	RigVMSchema = RigVMController->GetGraph()->GetSchema();
-	check(RigVMSchema);
-	EditorData = Graph->GetTypedOuter<UAnimNextRigVMAssetEditorData>();
-	check(EditorData);
+	ContextData.RigVMClientHost = ContextData.Graph->GetImplementingOuter<IRigVMClientHost>();
+	check(ContextData.RigVMClientHost);
+	ContextData.RigVMHost = ContextData.Graph->GetTypedOuter<URigVMHost>();
+	check(ContextData.RigVMHost);
+	ContextData.RigVMController = ContextData.RigVMClientHost->GetRigVMClient()->GetController(ContextData.Graph);
+	check(ContextData.RigVMController);
+	ContextData.RigVMSchema = ContextData.RigVMController->GetGraph()->GetSchema();
+	check(ContextData.RigVMSchema);
+	ContextData.EditorData = ContextData.Graph->GetTypedOuter<UAnimNextRigVMAssetEditorData>();
+	check(ContextData.EditorData);
 
 	SBorder::Construct(SBorder::FArguments()
 		.BorderImage(FAppStyle::Get().GetBrush("Menu.Background"))
@@ -214,7 +89,7 @@ void SActionMenu::Construct(const FArguments& InArgs, UEdGraph* InGraph)
 							return SNew(SExpanderArrow, InActionMenuData.TableRow);
 						})
 						.DraggedFromPins(DraggedFromPins)
-						.GraphObj(Graph)
+						.GraphObj(ContextData.Graph)
 						.AlphaSortItems(true)
 						.bAllowPreselectedItemActivation(true)
 				]
@@ -225,12 +100,12 @@ void SActionMenu::Construct(const FArguments& InArgs, UEdGraph* InGraph)
 
 void SActionMenu::CollectAllActions(FGraphActionListBuilderBase& OutAllActions)
 {
-	if (!Graph)
+	if (!ContextData.Graph)
 	{
 		return;
 	}
 
-	FGraphContextMenuBuilder MenuBuilder(Graph);
+	FGraphContextMenuBuilder MenuBuilder(ContextData.Graph);
 	if (!DraggedFromPins.IsEmpty())
 	{
 		MenuBuilder.FromPin = DraggedFromPins[0];
@@ -291,7 +166,7 @@ TSharedRef<SWidget> SActionMenu::OnCreateWidgetForAction(FCreateWidgetForActionD
 
 void SActionMenu::OnActionSelected(const TArray<TSharedPtr<FEdGraphSchemaAction>>& SelectedAction, ESelectInfo::Type InSelectionType)
 {
-	if (!Graph)
+	if (!ContextData.Graph)
 	{
 		return;
 	}
@@ -303,7 +178,7 @@ void SActionMenu::OnActionSelected(const TArray<TSharedPtr<FEdGraphSchemaAction>
 
 	for (const TSharedPtr<FEdGraphSchemaAction>& Action : SelectedAction)
 	{
-		if (Action.IsValid() && Graph)
+		if (Action.IsValid() && ContextData.Graph)
 		{
 			if (!bActionExecuted && (Action->GetTypeId() != FEdGraphSchemaAction_Dummy::StaticGetTypeId()))
 			{
@@ -311,7 +186,7 @@ void SActionMenu::OnActionSelected(const TArray<TSharedPtr<FEdGraphSchemaAction>
 				bActionExecuted = true;
 			}
 
-			Action->PerformAction(Graph, DraggedFromPins, NewNodePosition);
+			Action->PerformAction(ContextData.Graph, DraggedFromPins, NewNodePosition);
 		}
 	}
 }
diff --git a/Engine/Plugins/Experimental/AnimNext/Source/AnimNextEditor/Private/Common/SActionMenu.h b/Engine/Plugins/Experimental/AnimNext/Source/AnimNextEditor/Private/Common/SActionMenu.h
deleted file mode 100644
diff --git a/Engine/Plugins/Experimental/AnimNext/Source/AnimNextEditor/Private/EditorUtils.cpp b/Engine/Plugins/Experimental/AnimNext/Source/AnimNextEditor/Private/EditorUtils.cpp
--- a/Engine/Plugins/Experimental/AnimNext/Source/AnimNextEditor/Private/EditorUtils.cpp
+++ b/Engine/Plugins/Experimental/AnimNext/Source/AnimNextEditor/Private/EditorUtils.cpp
@@ -15,6 +15,7 @@
 #include "String/ParseTokens.h"
 #include "Variables/IVariableBindingType.h"
 #include "Widgets/SNullWidget.h"
+#include "Common/GraphEditorSchemaActions.h"
 
 #define LOCTEXT_NAMESPACE "AnimNextEditorUtils"
 
@@ -359,6 +360,39 @@ bool FUtils::IsValidParameterName(const FName InName, FText& OutErrorText)
 	return true;
 }
 
+bool FUtils::AddSchemaRigUnitAction(const TSubclassOf<URigVMUnitNode>& UnitNodeClass, UScriptStruct* Struct, const FRigVMFunction& Function, FGraphContextMenuBuilder& InContexMenuBuilder)
+{
+	FString CategoryMetadata, DisplayNameMetadata, MenuDescSuffixMetadata;
+	UE::AnimNext::Editor::FUtils::GetRigUnitStructMetadata(Struct, CategoryMetadata, DisplayNameMetadata, MenuDescSuffixMetadata);
+	if (DisplayNameMetadata.IsEmpty())
+	{
+		DisplayNameMetadata = Function.GetMethodName().ToString();
+	}
+	if (!MenuDescSuffixMetadata.IsEmpty())
+	{
+		MenuDescSuffixMetadata = TEXT(" ") + MenuDescSuffixMetadata;
+	}
+	const FText NodeCategory = FText::FromString(CategoryMetadata);
+	const FText MenuDesc = FText::FromString(DisplayNameMetadata + MenuDescSuffixMetadata);
+	const FText ToolTip = Struct->GetToolTipText();
+
+	if (MenuDesc.IsEmpty())
+	{
+		return false;
+	}
+
+	InContexMenuBuilder.AddAction(MakeShared<FAnimNextSchemaAction_RigUnit>(UnitNodeClass, Struct, NodeCategory, MenuDesc, ToolTip));
+	return true;
 }
 
+void FUtils::GetRigUnitStructMetadata(const UScriptStruct* Struct, FString& OutCategoryMetadata, FString& OutDisplayNameMetadata, FString& OutMenuDescSuffixMetadata)
+{
+	check(Struct);
+	Struct->GetStringMetaDataHierarchical(FRigVMStruct::CategoryMetaName, &OutCategoryMetadata);
+	Struct->GetStringMetaDataHierarchical(FRigVMStruct::DisplayNameMetaName, &OutDisplayNameMetadata);
+	Struct->GetStringMetaDataHierarchical(FRigVMStruct::MenuDescSuffixMetaName, &OutMenuDescSuffixMetadata);
+}
+
+} // end namespace
+
 #undef LOCTEXT_NAMESPACE
\ No newline at end of file
diff --git a/Engine/Plugins/Experimental/AnimNext/Source/AnimNextEditor/Private/EditorUtils.h b/Engine/Plugins/Experimental/AnimNext/Source/AnimNextEditor/Private/EditorUtils.h
deleted file mode 100644
diff --git a/Engine/Plugins/Experimental/AnimNext/Source/AnimNextEditor/Public/IAnimNextEditorModule.h b/Engine/Plugins/Experimental/AnimNext/Source/AnimNextEditor/Public/IAnimNextEditorModule.h
--- a/Engine/Plugins/Experimental/AnimNext/Source/AnimNextEditor/Public/IAnimNextEditorModule.h
+++ b/Engine/Plugins/Experimental/AnimNext/Source/AnimNextEditor/Public/IAnimNextEditorModule.h
@@ -8,6 +8,12 @@
 class SWidget;
 class IMessageLogListing;
 struct FTopLevelAssetPath;
+struct FGraphContextMenuBuilder;
+
+namespace UE::Workspace
+{
+	class IWorkspaceEditor;
+}
 
 namespace UE::AnimNext::Editor
 {
@@ -16,10 +22,14 @@ const FLazyName CompilerResultsTabName("CompilerResultsTab");
 const FLazyName LogListingName("AnimNextCompilerResults");
 
 struct FVariablePickerArgs;
+struct FActionMenuContextData;
 
 class IAnimNextEditorModule : public IModuleInterface
 {
 public:
+	DECLARE_MULTICAST_DELEGATE_ThreeParams(FCollectGraphMenuActionsMulticast, const TWeakPtr<UE::Workspace::IWorkspaceEditor>& /*IWorkspaceEditor*/, FGraphContextMenuBuilder& /*InContexMenuBuilder*/, const FActionMenuContextData& /*ActionMenuContextData*/);
+	typedef FCollectGraphMenuActionsMulticast::FDelegate FOnCollectGraphMenuActionsDelegate;
+
 	// Register a valid fragment type name to be used with parameter UOLs
 	// @param InLocatorFragmentEditorName The name of the locator fragment editor
 	virtual void RegisterLocatorFragmentEditorType(FName InLocatorFragmentEditorName) = 0;
@@ -35,6 +45,15 @@ public:
 	// Remove a UClass path to the set of classes which can be opened within an AnimNext Workspace
 	// @param InClassAssetPath Asset path for to-be-unregistered Class 
 	virtual void RemoveWorkspaceSupportedAssetClass(const FTopLevelAssetPath& InClassAssetPath) = 0;
+
+	// Register a graph context menu actions provider
+	// @param InCollectDelegate Delegate to add
+	// @return FDelegateHandle
+	virtual FDelegateHandle RegisterGraphMenuActionsProvider(const FOnCollectGraphMenuActionsDelegate& InCollectDelegate) = 0;
+
+	// Unregister a graph context menu actions provider
+	// @param InDelegateHandle Handle of the delegate to remove
+	virtual void UnregisterGraphMenuActionsProvider(const FDelegateHandle& InDelegateHandle) = 0;
 };
 
 }
diff --git a/Engine/Plugins/Experimental/AnimNext/Source/AnimNextUncookedOnly/Internal/AnimNextRigVMAssetEditorData.h b/Engine/Plugins/Experimental/AnimNext/Source/AnimNextUncookedOnly/Internal/AnimNextRigVMAssetEditorData.h
--- a/Engine/Plugins/Experimental/AnimNext/Source/AnimNextUncookedOnly/Internal/AnimNextRigVMAssetEditorData.h
+++ b/Engine/Plugins/Experimental/AnimNext/Source/AnimNextUncookedOnly/Internal/AnimNextRigVMAssetEditorData.h
@@ -8,10 +8,10 @@
 #include "RigVMBlueprint.h"
 #include "UncookedOnlyUtils.h"
 #include "Kismet/BlueprintFunctionLibrary.h"
+#include "Entries/AnimNextRigVMAssetEntry.h"
 #include "AnimNextRigVMAssetEditorData.generated.h"
 
 enum class ERigVMGraphNotifType : uint8;
-class UAnimNextRigVMAssetEntry;
 class UAnimNextRigVMAssetEditorData;
 class UAnimNextEdGraph;
 class UAnimNextDataInterfaceEntry;
@@ -42,7 +42,7 @@ namespace UE::AnimNext::Editor
 	class SAddVariablesDialog;
 	class FVariableProxyCustomization;
 	struct FPublicVariablesImpl;
-	class SActionMenu;
+	class FAnimNextAnimGraphEditorModule;
 	class FCallFunctionSharedDataDetails;
 }
 
@@ -180,7 +180,7 @@ protected:
 	friend class UE::AnimNext::Editor::FVariableProxyCustomization;
 	friend class UAnimNextDataInterfaceEntry;
 	friend struct UE::AnimNext::Editor::FPublicVariablesImpl;
-	friend UE::AnimNext::Editor::SActionMenu;
+	friend UE::AnimNext::Editor::FAnimNextAnimGraphEditorModule;
 	friend class UE::AnimNext::Editor::FCallFunctionSharedDataDetails;
 
 	// UObject interface
@@ -196,6 +196,8 @@ protected:
 	void HandlePackageDone(const FEndLoadPackageContext& Context);
 	void HandlePackageDone();
 
+	virtual void GetAnimNextAssetRegistryTags(FAssetRegistryTagsContext& Context, FAnimNextAssetRegistryExports& OutExports) const {}
+
 	// IRigVMClientHost interface
 	virtual FString GetAssetName() const override { return GetName(); }
 	virtual UClass* GetRigVMSchemaClass() const override;
diff --git a/Engine/Plugins/Experimental/AnimNext/Source/AnimNextUncookedOnly/Internal/AnimNextUnitNode.h b/Engine/Plugins/Experimental/AnimNext/Source/AnimNextUncookedOnly/Internal/AnimNextUnitNode.h
--- a/Engine/Plugins/Experimental/AnimNext/Source/AnimNextUncookedOnly/Internal/AnimNextUnitNode.h
+++ b/Engine/Plugins/Experimental/AnimNext/Source/AnimNextUncookedOnly/Internal/AnimNextUnitNode.h
@@ -16,5 +16,9 @@ class UAnimNextUnitNode : public URigVMUnitNode
 	GENERATED_BODY()
 
 public:
-	virtual bool HasNonNativePins() const { return true; }
+	// Override node functions
+	virtual bool HasNonNativePins() const
+	{
+		return true;
+	}
 };
diff --git a/Engine/Plugins/Experimental/AnimNext/Source/AnimNextUncookedOnly/Internal/Entries/AnimNextVariableEntry.h b/Engine/Plugins/Experimental/AnimNext/Source/AnimNextUncookedOnly/Internal/Entries/AnimNextVariableEntry.h
--- a/Engine/Plugins/Experimental/AnimNext/Source/AnimNextUncookedOnly/Internal/Entries/AnimNextVariableEntry.h
+++ b/Engine/Plugins/Experimental/AnimNext/Source/AnimNextUncookedOnly/Internal/Entries/AnimNextVariableEntry.h
@@ -30,7 +30,7 @@ namespace UE::AnimNext::Editor
 	class SVariablesOutlinerValue;
 	class FVariableProxyCustomization;
 	struct FPublicVariablesImpl;
-	class SActionMenu;
+	class FAnimNextEditorModule;
 }
 
 namespace UE::AnimNext::Tests
@@ -69,7 +69,7 @@ class UAnimNextVariableEntry : public UAnimNextRigVMAssetEntry, public IAnimNext
 	friend struct UE::AnimNext::UncookedOnly::FUtils;
 	friend class UE::AnimNext::Editor::FVariableProxyCustomization;
 	friend struct UE::AnimNext::Editor::FPublicVariablesImpl;
-	friend class UE::AnimNext::Editor::SActionMenu;
+	friend class UE::AnimNext::Editor::FAnimNextEditorModule;
 
 	UAnimNextVariableEntry();
 
diff --git a/Engine/Plugins/Experimental/AnimNext/Source/AnimNextUncookedOnly/Internal/Module/AnimNextModule_EditorData.h b/Engine/Plugins/Experimental/AnimNext/Source/AnimNextUncookedOnly/Internal/Module/AnimNextModule_EditorData.h
--- a/Engine/Plugins/Experimental/AnimNext/Source/AnimNextUncookedOnly/Internal/Module/AnimNextModule_EditorData.h
+++ b/Engine/Plugins/Experimental/AnimNext/Source/AnimNextUncookedOnly/Internal/Module/AnimNextModule_EditorData.h
@@ -46,6 +46,7 @@ class UAnimNextModule_EditorData : public UAnimNextDataInterface_EditorData
 	friend class FAnimationAnimNextEditorTest_GraphTraitOperations;
 	friend class FAnimationAnimNextRuntimeTest_GraphExecute;
 	friend class FAnimationAnimNextRuntimeTest_GraphExecuteLatent;
+	friend class FAnimationAnimNextEditorTest_GraphManifest;
 
 private:
 	// UAnimNextRigVMAssetEditorData interface
diff --git a/Engine/Plugins/Experimental/AnimNext/Source/AnimNextUncookedOnly/Internal/UncookedOnlyUtils.h b/Engine/Plugins/Experimental/AnimNext/Source/AnimNextUncookedOnly/Internal/UncookedOnlyUtils.h
--- a/Engine/Plugins/Experimental/AnimNext/Source/AnimNextUncookedOnly/Internal/UncookedOnlyUtils.h
+++ b/Engine/Plugins/Experimental/AnimNext/Source/AnimNextUncookedOnly/Internal/UncookedOnlyUtils.h
@@ -10,6 +10,9 @@
 #include "RigVMCore/RigVMTemplate.h"
 #include "StructUtils/InstancedStruct.h"
 #include "RigVMCore/RigVMGraphFunctionDefinition.h"
+#include "UObject/SoftObjectPath.h"
+#include "RigVMModel/RigVMGraph.h"
+// --- ---
 #include "UncookedOnlyUtils.generated.h"
 
 struct FAnimNextVariableBindingData;
@@ -19,6 +22,8 @@ struct FWorkspaceOutlinerItemExport;
 struct FRigVMGraphFunctionData;
 struct FRigVMCompileSettings;
 struct FRigVMGraphFunctionHeaderArray;
+struct FGraphContextMenuBuilder;
+struct FRigVMFunction;
 class UAnimNextModule;
 class UAnimNextModule_EditorData;
 class UAnimNextEdGraph;
@@ -28,6 +33,7 @@ class UAnimNextEdGraph;
 class UAnimNextRigVMAsset;
 class UAnimNextRigVMAssetEditorData;
 class UAnimNextRigVMAssetEntry;
+class URigVMUnitNode;
 
 
 namespace UE
@@ -91,6 +97,42 @@ struct FAnimNextAssetRegistryExportedVariable
 	uint32 Flags = (int32)EAnimNextExportedVariableFlags::NoFlags;
 };
 
+USTRUCT()
+struct FAnimNextAssetRegistryManifestNode
+{
+	GENERATED_BODY()
+
+	FAnimNextAssetRegistryManifestNode() = default;
+
+	FAnimNextAssetRegistryManifestNode(const TObjectPtr<URigVMGraph>& InModelGraph, const FString& InNodeName, const FString& InNodeCategory, const FString& InMenuDesc, const FString& InToolTip)
+		: ModelGraph(FSoftObjectPath::ConstructFromObject(InModelGraph))
+		, NodeName(InNodeName)
+		, NodeCategory(InNodeCategory)
+		, MenuDesc(InMenuDesc)
+		, ToolTip(InToolTip)
+	{}
+
+	bool operator==(const FAnimNextAssetRegistryManifestNode& Other) const
+	{
+		return ModelGraph == Other.ModelGraph && NodeName == Other.NodeName;
+	}
+
+	UPROPERTY()
+	FSoftObjectPath ModelGraph;
+
+	UPROPERTY()
+	FString NodeName;
+
+	UPROPERTY()
+	FString NodeCategory;
+
+	UPROPERTY()
+	FString MenuDesc;
+
+	UPROPERTY()
+	FString ToolTip;
+};
+
 USTRUCT()
 struct FAnimNextAssetRegistryExports
 {
@@ -98,6 +140,9 @@ struct FAnimNextAssetRegistryExports
 
 	UPROPERTY()
 	TArray<FAnimNextAssetRegistryExportedVariable> Variables;
+
+	UPROPERTY()
+	TArray<FAnimNextAssetRegistryManifestNode> ManifestNodes;
 };
 
 namespace UE::AnimNext::UncookedOnly
@@ -191,7 +236,7 @@ struct ANIMNEXTUNCOOKEDONLY_API FUtils
 	static const FText& GetFunctionLibraryDisplayName();
 
 #if WITH_EDITOR
-	static void OpenProgrammaticGraphs(UAnimNextRigVMAssetEditorData* EditorData, const TArray<URigVMGraph*>& ProgrammaticGraphs);
+	static void OpenProgrammaticGraphs(UAnimNextRigVMAssetEditorData* InEditorData, const TArray<URigVMGraph*>& ProgrammaticGraphs);
 #endif // WITH_EDITOR
 
 	// Make an event name that we use as a wrapper to call RigVM functions 
diff --git a/Engine/Plugins/Experimental/AnimNext/Source/AnimNextUncookedOnly/Private/AnimNextRigVMAssetEditorData.cpp b/Engine/Plugins/Experimental/AnimNext/Source/AnimNextUncookedOnly/Private/AnimNextRigVMAssetEditorData.cpp
--- a/Engine/Plugins/Experimental/AnimNext/Source/AnimNextUncookedOnly/Private/AnimNextRigVMAssetEditorData.cpp
+++ b/Engine/Plugins/Experimental/AnimNext/Source/AnimNextUncookedOnly/Private/AnimNextRigVMAssetEditorData.cpp
@@ -255,11 +255,12 @@ void UAnimNextRigVMAssetEditorData::GetAssetRegistryTags(FAssetRegistryTagsConte
 
 	{
 		// We may not have compiled yet, so cache exports if we havent already
-		if(!CachedExports.IsSet())
+		if (!CachedExports.IsSet())
 		{
 			CachedExports = FAnimNextAssetRegistryExports();
 			FAnimNextAssetRegistryExports& OutExports = CachedExports.GetValue();
 
+			GetAnimNextAssetRegistryTags(Context, OutExports);
 			UE::AnimNext::UncookedOnly::FUtils::GetAssetVariables(this, OutExports);
 		}
 
@@ -1004,6 +1005,11 @@ void UAnimNextRigVMAssetEditorData::HandleModifiedEvent(ERigVMGraphNotifType InN
 			}
 			break;
 		}
+	case ERigVMGraphNotifType::PinRemoved:
+		{
+			RequestAutoVMRecompilation(); // can not check if it is a trait pin, as it has been already removed
+			break;
+		}
 	}
 	
 	// if the notification still has to be sent...
diff --git a/Engine/Plugins/Experimental/AnimNextAnimGraph/Source/AnimNextAnimGraphEditor/Private/AnimNextAnimGraphEditorModule.cpp b/Engine/Plugins/Experimental/AnimNextAnimGraph/Source/AnimNextAnimGraphEditor/Private/AnimNextAnimGraphEditorModule.cpp
--- a/Engine/Plugins/Experimental/AnimNextAnimGraph/Source/AnimNextAnimGraphEditor/Private/AnimNextAnimGraphEditorModule.cpp
+++ b/Engine/Plugins/Experimental/AnimNextAnimGraph/Source/AnimNextAnimGraphEditor/Private/AnimNextAnimGraphEditorModule.cpp
@@ -10,7 +10,7 @@
 #include "IWorkspaceEditorModule.h"
 #include "PropertyEditorModule.h"
 #include "TraitStackEditor.h"
-#include "UncookedOnlyUtils.h"
+#include "AnimGraphUncookedOnlyUtils.h"
 #include "Entries/AnimNextAnimationGraphEntry.h"
 #include "Features/IModularFeatures.h"
 #include "Framework/Docking/LayoutExtender.h"
@@ -25,6 +25,14 @@
 #include "Traits/AnimNextCallFunctionSharedDataDetails.h"
 #include "Traits/CallFunction.h"
 #include "WorkflowOrientedApp/WorkflowTabManager.h"
+#include "IWorkspaceEditor.h"
+#include "Graph/AnimGraphEditorSchemaActions.h"
+#include "Graph/AnimNextAnimationGraphSchema.h"
+#include "Common/SActionMenu.h"
+#include "Graph/RigUnit_AnimNextTraitStack.h"
+#include "AnimNextTraitStackUnitNode.h"
+#include "RigVMCore/RigVMFunction.h"
+#include "EditorUtils.h"
 
 #define LOCTEXT_NAMESPACE "FAnimNextAnimGraphEditorModule"
 
@@ -83,6 +91,7 @@ class FAnimNextAnimGraphEditorModule : public IModuleInterface
 
 		IAnimNextEditorModule& AnimNextEditorModule = FModuleManager::LoadModuleChecked<IAnimNextEditorModule>("AnimNextEditor");
 		AnimNextEditorModule.AddWorkspaceSupportedAssetClass(UAnimNextAnimationGraph::StaticClass()->GetClassPathName());
+		CollectMenuActionsDelegateHandle = AnimNextEditorModule.RegisterGraphMenuActionsProvider(IAnimNextEditorModule::FOnCollectGraphMenuActionsDelegate::CreateStatic(CollectContextMenuActions));
 
 		TraitStackEditor = MakeShared<FTraitStackEditor>();
 		IModularFeatures::Get().RegisterModularFeature(ITraitStackEditor::ModularFeatureName, TraitStackEditor.Get());
@@ -120,11 +129,77 @@ class FAnimNextAnimGraphEditorModule : public IModuleInterface
 		IModularFeatures::Get().UnregisterModularFeature(ITraitStackEditor::ModularFeatureName, TraitStackEditor.Get());
 		TraitStackEditor.Reset();
 
+		if (IAnimNextEditorModule* AnimNextEditorModule = FModuleManager::GetModulePtr<IAnimNextEditorModule>("AnimNextEditor"))
+		{
+			AnimNextEditorModule->UnregisterGraphMenuActionsProvider(CollectMenuActionsDelegateHandle);
+		}
+
 		FEdGraphUtilities::UnregisterVisualNodeFactory(AnimNextGraphPanelNodeFactory);
 		
 		FAnimationGraphMenuExtensions::UnregisterMenus();
 	}
 
+	static void CollectContextMenuActions(const TWeakPtr<UE::Workspace::IWorkspaceEditor>& InWorkspaceEditorWeak, FGraphContextMenuBuilder& InContexMenuBuilder, const FActionMenuContextData& ActionMenuContextData)
+	{
+		if (const URigVMEdGraph* RigVMEdGraph = Cast<URigVMEdGraph>(InContexMenuBuilder.CurrentGraph))
+		{
+			if (RigVMEdGraph->GetModel()->GetSchemaClass() == UAnimNextAnimationGraphSchema::StaticClass())
+			{
+				TArray<FAnimNextAssetRegistryExports> ManifestExports;
+				UE::AnimNext::UncookedOnly::FAnimGraphUtils::GetExportedManifestNodesFromAssetRegistry(ManifestExports);
+
+				TArray<FAssetData> WorkspaceAssets;
+				if (TSharedPtr<UE::Workspace::IWorkspaceEditor> WorkspaceEditor = InWorkspaceEditorWeak.Pin())
+				{
+					WorkspaceEditor->GetAssets(WorkspaceAssets);
+				}
+
+				for (const FAnimNextAssetRegistryExports& ManifestExport : ManifestExports)
+				{
+					for (const FAnimNextAssetRegistryManifestNode& ManifestNodeData : ManifestExport.ManifestNodes)
+					{
+						bool bIncludeManifestNode = ActionMenuContextData.bShowGlobalManifestNodes;
+
+						if (!ActionMenuContextData.bShowGlobalManifestNodes)
+						{
+							if (WorkspaceAssets.Num() > 0)
+							{
+								for (const FAssetData& WorkspaceAsset : WorkspaceAssets)
+								{
+									if (WorkspaceAsset.PackageName != ManifestNodeData.ModelGraph.GetLongPackageName())
+									{
+										continue;
+									}
+									bIncludeManifestNode = true;
+								}
+							}
+						}
+
+						if (bIncludeManifestNode)
+						{
+							InContexMenuBuilder.AddAction(MakeShared<FAnimNextSchemaAction_AddManifestNode>(ManifestNodeData));
+						}
+					}
+				}
+			}
+		}
+
+		// Add trait stack using a custom RigUnit node class
+		UScriptStruct* Struct = FRigUnit_AnimNextTraitStack::StaticStruct();
+		const FString FunctionName = FString::Printf(TEXT("%s::%s"), *Struct->GetStructCPPName(), *FRigVMStruct::ExecuteName.ToString());
+		const FRigVMFunction* Function = FRigVMRegistry::Get().FindFunction(*FunctionName);
+		if (ensure(Function != nullptr))
+		{
+			if (ActionMenuContextData.RigVMSchema != nullptr && ActionMenuContextData.RigVMSchema->SupportsUnitFunction(ActionMenuContextData.RigVMController, Function))
+			{
+				UE::AnimNext::Editor::FUtils::AddSchemaRigUnitAction(UAnimNextTraitStackUnitNode::StaticClass(), Struct, *Function, InContexMenuBuilder);
+			}
+		}
+	}
+
+	/** Graph context menu collect actions delegate handle */
+	FDelegateHandle CollectMenuActionsDelegateHandle;
+
 	/** Trait stack editor modular feature */
 	TSharedPtr<FTraitStackEditor> TraitStackEditor;
 
@@ -132,7 +207,7 @@ class FAnimNextAnimGraphEditorModule : public IModuleInterface
 	TSharedPtr<FAnimNextGraphPanelNodeFactory> AnimNextGraphPanelNodeFactory;
 };
 
-}
+} // end namespace
 
 #undef LOCTEXT_NAMESPACE
 
diff --git a/Engine/Plugins/Experimental/AnimNextAnimGraph/Source/AnimNextAnimGraphEditor/Private/Graph/AnimNextAnimationGraphMenuExtensions.cpp b/Engine/Plugins/Experimental/AnimNextAnimGraph/Source/AnimNextAnimGraphEditor/Private/Graph/AnimNextAnimationGraphMenuExtensions.cpp
--- a/Engine/Plugins/Experimental/AnimNextAnimGraph/Source/AnimNextAnimGraphEditor/Private/Graph/AnimNextAnimationGraphMenuExtensions.cpp
+++ b/Engine/Plugins/Experimental/AnimNextAnimGraph/Source/AnimNextAnimGraphEditor/Private/Graph/AnimNextAnimationGraphMenuExtensions.cpp
@@ -1,4 +1,4 @@
-﻿// Copyright Epic Games, Inc. All Rights Reserved.
+// Copyright Epic Games, Inc. All Rights Reserved.
 
 #include "AnimNextAnimationGraphMenuExtensions.h"
 
@@ -17,6 +17,10 @@
 #include "RigVMModel/RigVMController.h"
 #include "RigVMModel/Nodes/RigVMUnitNode.h"
 #include "TraitCore/TraitRegistry.h"
+#include "AnimNextController.h"
+#include "AnimNextTraitStackUnitNode.h"
+#include "AnimNextRigVMAssetEditorData.h"
+#include "RigVMModel/RigVMClient.h"
 
 #define LOCTEXT_NAMESPACE "FAnimationGraphMenuExtensions"
 
@@ -233,6 +237,11 @@ void FAnimationGraphMenuExtensions::RegisterMenus()
 					FNewToolMenuDelegate::CreateLambda(BuildExposeVariablesContextMenu));
 			}
 		}
+
+		if (UAnimNextTraitStackUnitNode* AnimNextUnitNode = Cast<UAnimNextTraitStackUnitNode>(ModelNode))
+		{
+			AddManifestSection(InMenu, AnimNextEdGraphNode, ModelNode);
+		}
 	}));
 }
 
@@ -244,6 +253,70 @@ void FAnimationGraphMenuExtensions::UnregisterMenus()
 	}
 }
 
+void FAnimationGraphMenuExtensions::AddManifestSection(UToolMenu* InMenu, const UAnimNextEdGraphNode* AnimNextEdGraphNode, URigVMNode* ModelNode)
+{
+	FToolMenuSection& Section = InMenu->AddSection("AnimNextManifestNodeActions", LOCTEXT("AnimNextManifestNodeActionsMenuHeader", "Manifest"));
+
+	const TWeakObjectPtr<const UAnimNextEdGraphNode> EdNodeWeak = AnimNextEdGraphNode;
+
+	if (!UE::AnimNext::UncookedOnly::FAnimGraphUtils::IsExposedToManifest(ModelNode))
+	{
+		Section.AddMenuEntry(
+			"AddTraitToManifest",
+			LOCTEXT("AddTraitToManifest", "Add to Manifest"),
+			LOCTEXT("AddTraitToManifest_Tooltip", "Adds this Trait Stack to the Manifest."),
+			FSlateIcon(),
+			FUIAction(FExecuteAction::CreateLambda([EdNodeWeak]()
+				{
+					if (const TStrongObjectPtr<const UAnimNextEdGraphNode> EdGraphNode = EdNodeWeak.Pin())
+					{
+						if (URigVMNode* ModelNode = EdGraphNode->GetModelNode())
+						{
+							UAnimNextController* Controller = CastChecked<UAnimNextController>(EdGraphNode->GetController());
+
+							if (ensure(UE::AnimNext::UncookedOnly::FAnimGraphUtils::IsExposedToManifest(ModelNode) == false))
+							{
+								FRigVMControllerCompileBracketScope CompileScope(Controller);
+
+								Controller->AddNodeToManifest(ModelNode);
+								
+								if (UAnimNextRigVMAssetEditorData* EditorData = EdGraphNode->GetGraph()->GetTypedOuter<UAnimNextRigVMAssetEditorData>())
+								{
+									UE::AnimNext::UncookedOnly::FAnimGraphUtils::RequestVMAutoRecompile(EditorData);  // required to force an asset tags update
+								}
+							}
+						}								
+					}
+				}))
+			);
+	}
+	else
+	{
+		Section.AddMenuEntry(
+			"RemoveTraitFromManifest",
+			LOCTEXT("RemoveTraitFromManifest", "Remove from Manifest"),
+			LOCTEXT("RemoveTraitFromManifest_Tooltip", "Remvoes this Trait Stack from the Manifest."),
+			FSlateIcon(),
+			FUIAction(FExecuteAction::CreateLambda([EdNodeWeak]()
+				{
+					if (const TStrongObjectPtr<const UAnimNextEdGraphNode> EdGraphNode = EdNodeWeak.Pin())
+					{
+						UAnimNextController* Controller = CastChecked<UAnimNextController>(EdGraphNode->GetController());
+
+						FRigVMControllerCompileBracketScope CompileScope(Controller);
+
+						Controller->RemoveNodeFromManifest(EdGraphNode->GetModelNode());
+						
+						if (UAnimNextRigVMAssetEditorData* EditorData = EdGraphNode->GetGraph()->GetTypedOuter<UAnimNextRigVMAssetEditorData>())
+						{
+							UE::AnimNext::UncookedOnly::FAnimGraphUtils::RequestVMAutoRecompile(EditorData); // required to force an asset tags update
+						}
+					}
+				}))
+		);
+	}
+}
+
 }
 
 #undef LOCTEXT_NAMESPACE
\ No newline at end of file
diff --git a/Engine/Plugins/Experimental/AnimNextAnimGraph/Source/AnimNextAnimGraphEditor/Private/Graph/AnimNextAnimationGraphMenuExtensions.h b/Engine/Plugins/Experimental/AnimNextAnimGraph/Source/AnimNextAnimGraphEditor/Private/Graph/AnimNextAnimationGraphMenuExtensions.h
--- a/Engine/Plugins/Experimental/AnimNextAnimGraph/Source/AnimNextAnimGraphEditor/Private/Graph/AnimNextAnimationGraphMenuExtensions.h
+++ b/Engine/Plugins/Experimental/AnimNextAnimGraph/Source/AnimNextAnimGraphEditor/Private/Graph/AnimNextAnimationGraphMenuExtensions.h
@@ -1,9 +1,14 @@
-﻿// Copyright Epic Games, Inc. All Rights Reserved.
+// Copyright Epic Games, Inc. All Rights Reserved.
 
 #pragma once
 
 #include "CoreMinimal.h"
 
+class UAnimNextEdGraphNode;
+class URigVMNode;
+class UToolMenu;
+struct FToolMenuSection;
+
 namespace UE::AnimNext::Editor
 {
 
@@ -11,6 +16,9 @@ struct FAnimationGraphMenuExtensions
 {
 	static void RegisterMenus();
 	static void UnregisterMenus();
+
+private:
+	static void AddManifestSection(UToolMenu* InMenu, const UAnimNextEdGraphNode* AnimNextEdGraphNode, URigVMNode* ModelNode);
 };
 
 }
diff --git a/Engine/Plugins/Experimental/AnimNextAnimGraph/Source/AnimNextAnimGraphEditor/Private/Graph/STraitStackView.cpp b/Engine/Plugins/Experimental/AnimNextAnimGraph/Source/AnimNextAnimGraphEditor/Private/Graph/STraitStackView.cpp
--- a/Engine/Plugins/Experimental/AnimNextAnimGraph/Source/AnimNextAnimGraphEditor/Private/Graph/STraitStackView.cpp
+++ b/Engine/Plugins/Experimental/AnimNextAnimGraph/Source/AnimNextAnimGraphEditor/Private/Graph/STraitStackView.cpp
@@ -18,6 +18,11 @@
 #include "DragAndDrop/DecoratedDragDropOp.h"
 #include "DragAndDrop/AssetDragDropOp.h"
 #include "Widgets/Input/SButton.h"
+#include "RigVMModel/RigVMController.h"
+#include "AnimNextTraitStackUnitNode.h"
+#include "AnimNextRigVMAssetEditorData.h"
+#include "RigVMModel/RigVMClient.h"
+#include "AnimGraphUncookedOnlyUtils.h"
 
 #define LOCTEXT_NAMESPACE "TraitListEditor"
 
@@ -76,7 +81,6 @@ void STraitStackView::Construct(const FArguments& InArgs, TSharedPtr<FTraitEdito
 		FCanExecuteAction::CreateSP(this, &STraitStackView::HasValidSelection));
 
 	const TSharedPtr<FTraitEditorSharedData>& TraitEditorSharedDataLocal = TraitEditorSharedData;
-	const TWeakObjectPtr<UAnimNextEdGraphNode> EdGraphNodeWeak = TraitEditorSharedData->EdGraphNodeWeak;
 
 	ChildSlot
 	[
@@ -111,24 +115,51 @@ void STraitStackView::Construct(const FArguments& InArgs, TSharedPtr<FTraitEdito
 							SNew(SInlineEditableTextBlock)
 							.Style(FAppStyle::Get(), "Graph.Node.NodeTitleInlineEditableText")
 							.Text_Lambda([this] { return TraitEditorSharedData->EdGraphNodeWeak.IsValid() ? TraitEditorSharedData->EdGraphNodeWeak->GetNodeTitle(ENodeTitleType::EditableTitle) : FText::GetEmpty(); })
-							.IsReadOnly_Lambda([EdGraphNodeWeak]()->bool
+							.IsReadOnly_Lambda([TraitEditorSharedDataLocal]()->bool
 							{
-								bool bIsReadOnly = false;
-								if (UAnimNextEdGraphNode* EdGraphNode = Cast<UAnimNextEdGraphNode>(EdGraphNodeWeak.Get()))
+								bool bIsReadOnly = true;
+								if (TraitEditorSharedDataLocal.IsValid())
 								{
-									bIsReadOnly = EdGraphNode->IsDeprecated() || EdGraphNode->IsOutDated();
+									if (UAnimNextEdGraphNode* EdGraphNode = Cast<UAnimNextEdGraphNode>(TraitEditorSharedDataLocal->EdGraphNodeWeak.Get()))
+									{
+										bIsReadOnly = EdGraphNode->IsDeprecated() || EdGraphNode->IsOutDated();
+									}
 								}
 								return bIsReadOnly;
 							})
-							.OnTextCommitted_Lambda([EdGraphNodeWeak](const FText& InNewText, ETextCommit::Type InCommitType)
+							.OnTextCommitted_Lambda([TraitEditorSharedDataLocal](const FText& InNewText, ETextCommit::Type InCommitType)
 							{
 								if(InCommitType == ETextCommit::OnEnter)
 								{
-									if(UAnimNextEdGraphNode* EdGraphNode = Cast<UAnimNextEdGraphNode>(EdGraphNodeWeak.Get()))
+									if (TraitEditorSharedDataLocal.IsValid())
 									{
-										// TODO zzz : Find the correct way to change the DisplayName / Title
-										//FScopedTransaction Transaction(LOCTEXT("SetNodeName", "Set Node title"));
-										//EdGraphNode->SetNodeTitle(InNewText);
+										if (UAnimNextEdGraphNode* EdGraphNode = Cast<UAnimNextEdGraphNode>(TraitEditorSharedDataLocal->EdGraphNodeWeak.Get()))
+										{
+											if (URigVMNode* ModelNode = EdGraphNode->GetModelNode())
+											{
+												if (UAnimNextTraitStackUnitNode* UnitNode = Cast<UAnimNextTraitStackUnitNode>(ModelNode))
+												{
+													FScopedTransaction Transaction(LOCTEXT("SetNodeTitle", "Set Node title"));
+
+													if (UScriptStruct* Struct = UnitNode->GetScriptStruct())
+													{
+														if (URigVMController* Controller = EdGraphNode->GetController())
+														{
+															if(Controller->SetNodeTitle(ModelNode, InNewText.ToString()))
+															{
+																if (UnitNode->IsExposedToManifest()) // if the node is exposed to manifest, we have to force a recompile to update asset tags
+																{
+																	if (UAnimNextRigVMAssetEditorData* EditorData = EdGraphNode->GetGraph()->GetTypedOuter<UAnimNextRigVMAssetEditorData>())
+																	{
+																		UE::AnimNext::UncookedOnly::FAnimGraphUtils::RequestVMAutoRecompile(EditorData);
+																	}
+																}
+															}
+														}
+													}
+												}
+											}
+										}
 									}
 								}
 							})
diff --git a/Engine/Plugins/Experimental/AnimNextAnimGraph/Source/AnimNextAnimGraphTestSuite/Private/AnimNextAnimGraphTraitGraphTest.cpp b/Engine/Plugins/Experimental/AnimNextAnimGraph/Source/AnimNextAnimGraphTestSuite/Private/AnimNextAnimGraphTraitGraphTest.cpp
--- a/Engine/Plugins/Experimental/AnimNextAnimGraph/Source/AnimNextAnimGraphTestSuite/Private/AnimNextAnimGraphTraitGraphTest.cpp
+++ b/Engine/Plugins/Experimental/AnimNextAnimGraph/Source/AnimNextAnimGraphTestSuite/Private/AnimNextAnimGraphTraitGraphTest.cpp
@@ -23,6 +23,8 @@
 #include "Misc/AutomationTest.h"
 #include "RigVMFunctions/Math/RigVMFunction_MathInt.h"
 #include "RigVMFunctions/Math/RigVMFunction_MathFloat.h"
+#include "AnimNextTraitStackUnitNode.h"
+#include "AnimGraphUncookedOnlyUtils.h"
 
 #if WITH_DEV_AUTOMATION_TESTS
 
@@ -374,6 +376,79 @@ bool FAnimationAnimNextEditorTest_GraphTraitOperations::RunTest(const FString& I
 	return true;
 }
 
+IMPLEMENT_SIMPLE_AUTOMATION_TEST(FAnimationAnimNextEditorTest_GraphManifest, "Animation.AnimNext.Editor.Graph.Manifest", EAutomationTestFlags::EditorContext | EAutomationTestFlags::EngineFilter)
+
+bool FAnimationAnimNextEditorTest_GraphManifest::RunTest(const FString& InParameters)
+{
+	using namespace UE::AnimNext;
+
+	{
+		FScopedClearNodeTemplateRegistry ScopedClearNodeTemplateRegistry;
+
+		UFactory* GraphFactory = NewObject<UAnimNextAnimationGraphFactory>();
+		UAnimNextAnimationGraph* AnimationGraph = CastChecked<UAnimNextAnimationGraph>(GraphFactory->FactoryCreateNew(UAnimNextAnimationGraph::StaticClass(), GetTransientPackage(), TEXT("TestAnimNextGraph"), RF_Transient, nullptr, nullptr, NAME_None));
+		UE_RETURN_ON_ERROR(AnimationGraph != nullptr, "FAnimationAnimNextEditorTest_GraphManifest -> Failed to create animation graph");
+
+		UAnimNextAnimationGraph_EditorData* EditorData = UncookedOnly::FUtils::GetEditorData<UAnimNextAnimationGraph_EditorData>(AnimationGraph);
+		UE_RETURN_ON_ERROR(EditorData != nullptr, "FAnimationAnimNextEditorTest_GraphManifest -> Failed to find module editor data");
+
+		UAnimNextController* AnimNextController = nullptr;
+		{
+			URigVMController* RigVMController = EditorData->GetRigVMClient()->GetController(EditorData->GetRigVMClient()->GetDefaultModel());
+			AnimNextController = Cast<UAnimNextController>(RigVMController);
+		}
+		UE_RETURN_ON_ERROR(AnimNextController != nullptr, "FAnimationAnimNextEditorTest_GraphManifest -> Failed to get RigVM controller");
+
+		// Create an empty trait stack node
+		UAnimNextTraitStackUnitNode* TraitStackNode = Cast<UAnimNextTraitStackUnitNode>(AnimNextController->AddUnitNode(FRigUnit_AnimNextTraitStack::StaticStruct(), UAnimNextTraitStackUnitNode::StaticClass(), FRigVMStruct::ExecuteName, FVector2D(0.0f, 0.0f), FString(), true, false));
+		UE_RETURN_ON_ERROR(TraitStackNode != nullptr, "FAnimationAnimNextEditorTest_GraphManifest -> Failed to create trait stack node");
+
+		UE_RETURN_ON_ERROR(TraitStackNode->IsExposedToManifest() == false, TEXT("FAnimationAnimNextEditorTest_GraphManifest -> Created node should not be at Manifest."));
+
+		// --- Add to Manifest ---
+		UE_RETURN_ON_ERROR(AnimNextController->AddNodeToManifestByName(TraitStackNode->GetFName()), TEXT("FAnimationAnimNextEditorTest_GraphManifest -> Add Node to Manifest Failed."));
+		UE_RETURN_ON_ERROR(TraitStackNode->IsExposedToManifest() == true, TEXT("FAnimationAnimNextEditorTest_GraphManifest -> Created node should be at Manifest."));
+		UE_RETURN_ON_ERROR(UE::AnimNext::UncookedOnly::FAnimGraphUtils::IsExposedToManifest(TraitStackNode) == true, TEXT("FAnimationAnimNextEditorTest_GraphManifest -> FAnimGraphUtils::IsExposedToManifest should return true."));
+
+		// --- Undo Add to Manifest ---
+		AnimNextController->Undo();
+		UE_RETURN_ON_ERROR(TraitStackNode->IsExposedToManifest() == false, TEXT("FAnimationAnimNextEditorTest_GraphManifest -> Created node should not be at Manifest after Undo."));
+		UE_RETURN_ON_ERROR(UE::AnimNext::UncookedOnly::FAnimGraphUtils::IsExposedToManifest(TraitStackNode) == false, TEXT("FAnimationAnimNextEditorTest_GraphManifest -> FAnimGraphUtils::IsExposedToManifest should return false after Undo."));
+
+		// --- Redo Add to Manifest ---
+		AnimNextController->Redo();
+		UE_RETURN_ON_ERROR(TraitStackNode->IsExposedToManifest() == true, TEXT("FAnimationAnimNextEditorTest_GraphManifest -> Created node should be at Manifest after Redo."));
+		UE_RETURN_ON_ERROR(UE::AnimNext::UncookedOnly::FAnimGraphUtils::IsExposedToManifest(TraitStackNode) == true, TEXT("FAnimationAnimNextEditorTest_GraphManifest -> FAnimGraphUtils::IsExposedToManifest should return true after Redo."));
+
+		// --- Check Regiostry Exports ---
+		FAnimNextAssetRegistryExports OutExports;
+		UE::AnimNext::UncookedOnly::FAnimGraphUtils::GetAssetManifestNodesRegistryExports(EditorData, OutExports);
+		UE_RETURN_ON_ERROR(OutExports.ManifestNodes.Num() == 1, TEXT("FAnimationAnimNextEditorTest_GraphManifest -> Asset Registry Exports Num should be 1."));
+
+		// --- Undo Add to Manifest ---
+		AnimNextController->Undo();
+		UE_RETURN_ON_ERROR(TraitStackNode->IsExposedToManifest() == false, TEXT("FAnimationAnimNextEditorTest_GraphManifest -> Created node should not be at Manifest after Undo."));
+		UE_RETURN_ON_ERROR(UE::AnimNext::UncookedOnly::FAnimGraphUtils::IsExposedToManifest(TraitStackNode) == false, TEXT("FAnimationAnimNextEditorTest_GraphManifest -> FAnimGraphUtils::IsExposedToManifest should return false after Undo."));
+
+		// --- Check Regiostry Exports ---
+		UE::AnimNext::UncookedOnly::FAnimGraphUtils::GetAssetManifestNodesRegistryExports(EditorData, OutExports);
+		UE_RETURN_ON_ERROR(OutExports.ManifestNodes.Num() == 0, TEXT("FAnimationAnimNextEditorTest_GraphManifest -> Asset Registry Exports Num should be 0 after Undo."));
+
+		// --- Redo Add to Manifest ---
+		AnimNextController->Redo();
+		UE_RETURN_ON_ERROR(TraitStackNode->IsExposedToManifest() == true, TEXT("FAnimationAnimNextEditorTest_GraphManifest -> Created node should be at Manifest after Redo."));
+		UE_RETURN_ON_ERROR(UE::AnimNext::UncookedOnly::FAnimGraphUtils::IsExposedToManifest(TraitStackNode) == true, TEXT("FAnimationAnimNextEditorTest_GraphManifest -> FAnimGraphUtils::IsExposedToManifest should return true after Redo."));
+
+		// --- Check Regiostry Exports ---
+		UE::AnimNext::UncookedOnly::FAnimGraphUtils::GetAssetManifestNodesRegistryExports(EditorData, OutExports);
+		UE_RETURN_ON_ERROR(OutExports.ManifestNodes.Num() == 1, TEXT("FAnimationAnimNextEditorTest_GraphManifest -> Asset Registry Exports Num should be 1 after Redo."));
+	}
+
+	Tests::FUtils::CleanupAfterTests();
+
+	return true;
+}
+
 IMPLEMENT_SIMPLE_AUTOMATION_TEST(FAnimationAnimNextRuntimeTest_GraphExecute, "Animation.AnimNext.Runtime.Graph.Execute", EAutomationTestFlags::EditorContext | EAutomationTestFlags::EngineFilter)
 
 bool FAnimationAnimNextRuntimeTest_GraphExecute::RunTest(const FString& InParameters)
diff --git a/Engine/Plugins/Experimental/AnimNextAnimGraph/Source/AnimNextAnimGraphUncookedOnly/AnimNextAnimGraphUncookedOnly.Build.cs b/Engine/Plugins/Experimental/AnimNextAnimGraph/Source/AnimNextAnimGraphUncookedOnly/AnimNextAnimGraphUncookedOnly.Build.cs
--- a/Engine/Plugins/Experimental/AnimNextAnimGraph/Source/AnimNextAnimGraphUncookedOnly/AnimNextAnimGraphUncookedOnly.Build.cs
+++ b/Engine/Plugins/Experimental/AnimNextAnimGraph/Source/AnimNextAnimGraphUncookedOnly/AnimNextAnimGraphUncookedOnly.Build.cs
@@ -31,6 +31,8 @@ namespace UnrealBuildTool.Rules
 						"UnrealEd",
 						"BlueprintGraph",
 						"WorkspaceEditor",
+						"AnimNextEditor",
+						"RigVMEditor",
 					}
 				);
 			}
diff --git a/Engine/Plugins/Experimental/AnimNextAnimGraph/Source/AnimNextAnimGraphUncookedOnly/Internal/AnimGraphUncookedOnlyUtils.h b/Engine/Plugins/Experimental/AnimNextAnimGraph/Source/AnimNextAnimGraphUncookedOnly/Internal/AnimGraphUncookedOnlyUtils.h
--- a/Engine/Plugins/Experimental/AnimNextAnimGraph/Source/AnimNextAnimGraphUncookedOnly/Internal/AnimGraphUncookedOnlyUtils.h
+++ b/Engine/Plugins/Experimental/AnimNextAnimGraph/Source/AnimNextAnimGraphUncookedOnly/Internal/AnimGraphUncookedOnlyUtils.h
@@ -1,11 +1,14 @@
-﻿// Copyright Epic Games, Inc. All Rights Reserved.
+// Copyright Epic Games, Inc. All Rights Reserved.
 
 #pragma once
 
 #include "CoreMinimal.h"
 
 class URigVMNode;
+class UAnimNextController;
 class URigVMController;
+class UAnimNextRigVMAssetEditorData;
+struct FAnimNextAssetRegistryExports;
 
 namespace UE::AnimNext::UncookedOnly
 {
@@ -16,7 +19,18 @@ struct ANIMNEXTANIMGRAPHUNCOOKEDONLY_API FAnimGraphUtils
 	static void SetupAnimGraph(const FName EntryName, URigVMController* InController, bool bSetupUndoRedo = true);
 
 	/** Check whether the supplied node is a trait stack */
-	static bool IsTraitStackNode(URigVMNode* InModelNode);
+	static bool IsTraitStackNode(const URigVMNode* InModelNode);
+
+	// Gets the exported manifest nodes that are used by a RigVM asset
+	static void GetAssetManifestNodesRegistryExports(const UAnimNextRigVMAssetEditorData* InEditorData, FAnimNextAssetRegistryExports& OutExports);
+
+	// Gets all the manifest node defs that are exported to the asset registry
+	static bool GetExportedManifestNodesFromAssetRegistry(TArray<FAnimNextAssetRegistryExports>& OutExports);
+
+	// Returns true if the node is exposed to the manifest
+	static bool IsExposedToManifest(const URigVMNode* InModelNode);
+
+	static bool RequestVMAutoRecompile(UAnimNextRigVMAssetEditorData* EditorData);
 };
 
 }
\ No newline at end of file
diff --git a/Engine/Plugins/Experimental/AnimNextAnimGraph/Source/AnimNextAnimGraphUncookedOnly/Internal/AnimNextController.h b/Engine/Plugins/Experimental/AnimNextAnimGraph/Source/AnimNextAnimGraphUncookedOnly/Internal/AnimNextController.h
--- a/Engine/Plugins/Experimental/AnimNextAnimGraph/Source/AnimNextAnimGraphUncookedOnly/Internal/AnimNextController.h
+++ b/Engine/Plugins/Experimental/AnimNextAnimGraph/Source/AnimNextAnimGraphUncookedOnly/Internal/AnimNextController.h
@@ -44,4 +44,16 @@ public:
 
 	// Adds a unit node with a dynamic number of pins
 	URigVMUnitNode* AddUnitNodeWithPins(UScriptStruct* InScriptStruct, const FRigVMPinInfoArray& PinArray, const FName& InMethodName = TEXT("Execute"), const FVector2D& InPosition = FVector2D::ZeroVector, const FString& InNodeName = TEXT(""), bool bSetupUndoRedo = true, bool bPrintPythonCommand = false);
+
+	UFUNCTION(BlueprintCallable, Category = RigVMController)
+	ANIMNEXTANIMGRAPHUNCOOKEDONLY_API bool AddNodeToManifest(URigVMNode* InModelNode, bool bSetupUndoRedo = true, bool bPrintPythonCommand = false);
+
+	UFUNCTION(BlueprintCallable, Category = RigVMController)
+	ANIMNEXTANIMGRAPHUNCOOKEDONLY_API bool AddNodeToManifestByName(const FName& InNodeName, bool bSetupUndoRedo = true, bool bMergeUndoAction = false);
+
+	UFUNCTION(BlueprintCallable, Category = RigVMController)
+	ANIMNEXTANIMGRAPHUNCOOKEDONLY_API bool RemoveNodeFromManifest(URigVMNode* InModelNode, bool bSetupUndoRedo = true, bool bPrintPythonCommand = false);
+
+	UFUNCTION(BlueprintCallable, Category = RigVMController)
+	ANIMNEXTANIMGRAPHUNCOOKEDONLY_API bool RemoveNodeFromManifestByName(const FName& InNodeName, bool bSetupUndoRedo = true, bool bPrintPythonCommand = false);
 };
diff --git a/Engine/Plugins/Experimental/AnimNextAnimGraph/Source/AnimNextAnimGraphUncookedOnly/Internal/AnimNextTraitStackUnitNode.h b/Engine/Plugins/Experimental/AnimNextAnimGraph/Source/AnimNextAnimGraphUncookedOnly/Internal/AnimNextTraitStackUnitNode.h
new file mode 100644
--- /dev/null
+++ b/Engine/Plugins/Experimental/AnimNextAnimGraph/Source/AnimNextAnimGraphUncookedOnly/Internal/AnimNextTraitStackUnitNode.h
@@ -0,0 +1,37 @@
+// Copyright Epic Games, Inc. All Rights Reserved.
+
+#pragma once
+
+#include "CoreMinimal.h"
+#include "AnimNextUnitNode.h"
+#include "AnimNextTraitStackUnitNode.generated.h"
+
+/**
+  * Implements AnimNext RigVM unit node extensions for Trait Stacks
+  */
+UCLASS(MinimalAPI)
+class UAnimNextTraitStackUnitNode : public URigVMUnitNode
+{
+	GENERATED_BODY()
+
+public:
+	// Override node functions
+	virtual FString GetNodeTitle() const override
+	{
+		return !NodeTitle.IsEmpty() ? NodeTitle : Super::GetNodeTitle();
+	}
+
+	bool IsExposedToManifest() const
+	{
+		return bIncludedInManifest;
+	}
+
+	void SetManifestState(bool bInIncludedInManifest)
+	{
+		bIncludedInManifest = bInIncludedInManifest;
+	}
+private:
+
+	UPROPERTY()
+	bool bIncludedInManifest = false;
+};
diff --git a/Engine/Plugins/Experimental/AnimNextAnimGraph/Source/AnimNextAnimGraphUncookedOnly/Internal/Graph/AnimGraphEditorSchemaActions.cpp b/Engine/Plugins/Experimental/AnimNextAnimGraph/Source/AnimNextAnimGraphUncookedOnly/Internal/Graph/AnimGraphEditorSchemaActions.cpp
new file mode 100644
--- /dev/null
+++ b/Engine/Plugins/Experimental/AnimNextAnimGraph/Source/AnimNextAnimGraphUncookedOnly/Internal/Graph/AnimGraphEditorSchemaActions.cpp
@@ -0,0 +1,105 @@
+// Copyright Epic Games, Inc. All Rights Reserved.
+
+#include "AnimGraphEditorSchemaActions.h"
+#include "EditorUtils.h"
+#include "AnimNextEdGraph.h"
+#include "AnimNextEdGraphNode.h"
+#include "EdGraphNode_Comment.h"
+#include "Editor.h"
+#include "Module/AnimNextModule_EditorData.h"
+#include "RigVMModel/Nodes/RigVMUnitNode.h"
+#include "Settings/ControlRigSettings.h"
+#include "Units/RigUnit.h"
+#include "EdGraphSchema_K2.h"
+#include "GraphEditor.h"
+#include "Framework/MultiBox/MultiBoxBuilder.h"
+#include "Framework/Application/SlateApplication.h"
+#include "UncookedOnlyUtils.h"
+#include "AnimNextController.h"
+#include "AnimGraphUncookedOnlyUtils.h"
+#include "Editor/RigVMEditorTools.h"
+#include "Exporters/Exporter.h"
+#include "UnrealExporter.h"
+#include "AnimNextUnitNode.h"
+
+#define LOCTEXT_NAMESPACE "AnimNextAnimGraphSchemaActions"
+
+namespace Private
+{
+static const FText ManifestMenuElementPrefix(LOCTEXT("ManifestMenuElementCategoryPrefix", "(M) - {0}"));
+}
+
+// --- Manifest Node ---
+
+FAnimNextSchemaAction_AddManifestNode::FAnimNextSchemaAction_AddManifestNode(const FAnimNextAssetRegistryManifestNode& InManifestNodeData, const FText& InKeywords)
+	: FAnimNextSchemaAction(FText::Format(Private::ManifestMenuElementPrefix, FText::FromString(InManifestNodeData.NodeCategory)), FText::FromString(InManifestNodeData.MenuDesc), FText::FromString(InManifestNodeData.ToolTip), InKeywords)
+	, ModelGraph(InManifestNodeData.ModelGraph)
+	, NodeName(InManifestNodeData.NodeName)
+{
+}
+
+UEdGraphNode* FAnimNextSchemaAction_AddManifestNode::PerformAction(UEdGraph* ParentGraph, TArray<UEdGraphPin*>& FromPins, const FVector2D Location, bool bSelectNewNode)
+{
+	IRigVMClientHost* Host = ParentGraph->GetImplementingOuter<IRigVMClientHost>();
+	URigVMEdGraphNode* NewNode = nullptr;
+	URigVMEdGraph* EdGraph = Cast<URigVMEdGraph>(ParentGraph);
+
+	UEdGraphPin* FromPin = nullptr;
+	if (FromPins.Num() > 0)
+	{
+		FromPin = FromPins[0];
+	}
+
+	if (Host != nullptr && EdGraph != nullptr)
+	{
+		if (UAnimNextController* Controller = Cast<UAnimNextController>(EdGraph->GetController()))
+		{
+			FName ValidName = UE::AnimNext::Editor::FUtils::ValidateName(Cast<UObject>(Host), NodeName);
+
+			Controller->OpenUndoBracket(FString::Printf(TEXT("Add '%s' Manifest Node"), *ValidName.ToString()));
+
+			const URigVMGraph* SourceGraph = ModelGraph.LoadSynchronous();
+			check(SourceGraph != nullptr);
+			URigVMNode* SourceNode = SourceGraph->FindNode(NodeName);
+
+			FStringOutputDevice Archive;
+			const FExportObjectInnerContext Context;
+
+			UExporter::ExportToOutputDevice(&Context, SourceNode, NULL, Archive, TEXT("copy"), 0, PPF_ExportsNotFullyQualified | PPF_Copy | PPF_Delimited, false, SourceNode->GetOuter());
+
+			const FString NodeData = MoveTemp(Archive);
+
+			TArray<FName> NodeNamesCreated = UE::RigVM::Editor::Tools::ImportNodesFromText(Location, NodeData, Controller, EdGraph->GetModel(), Host->GetLocalFunctionLibrary(), Host->GetRigVMGraphFunctionHost());
+
+			TArray<URigVMNode*> ModelNodes;
+			for (const FName& Name : NodeNamesCreated)
+			{
+				URigVMNode* Node = EdGraph->GetModel()->FindNodeByName(Name);
+				if (Node == nullptr)
+				{
+					return nullptr;
+				}
+				ModelNodes.AddUnique(Node);
+			}
+			if (URigVMNode* ModelNode = ModelNodes.Num() > 0 ? ModelNodes[0] : nullptr)
+			{
+				NewNode = Cast<URigVMEdGraphNode>(EdGraph->FindNodeForModelNodeName(ModelNode->GetFName()));
+				check(NewNode);
+
+				if (NewNode)
+				{
+					Controller->RemoveNodeFromManifest(ModelNode, false, false);
+				}
+				Controller->CloseUndoBracket();
+			}
+			else
+			{
+				Controller->CancelUndoBracket();
+			}
+		}
+	}
+
+	return NewNode;
+}
+
+#undef LOCTEXT_NAMESPACE
diff --git a/Engine/Plugins/Experimental/AnimNextAnimGraph/Source/AnimNextAnimGraphUncookedOnly/Internal/Graph/AnimGraphEditorSchemaActions.h b/Engine/Plugins/Experimental/AnimNextAnimGraph/Source/AnimNextAnimGraphUncookedOnly/Internal/Graph/AnimGraphEditorSchemaActions.h
new file mode 100644
--- /dev/null
+++ b/Engine/Plugins/Experimental/AnimNextAnimGraph/Source/AnimNextAnimGraphUncookedOnly/Internal/Graph/AnimGraphEditorSchemaActions.h
@@ -0,0 +1,25 @@
+// Copyright Epic Games, Inc. All Rights Reserved.
+
+#pragma once
+
+#include "CoreMinimal.h"
+#include "Common/GraphEditorSchemaActions.h"
+#include "UncookedOnlyUtils.h"
+#include "AnimGraphEditorSchemaActions.generated.h"
+
+USTRUCT()
+struct ANIMNEXTANIMGRAPHUNCOOKEDONLY_API FAnimNextSchemaAction_AddManifestNode : public FAnimNextSchemaAction
+{
+	GENERATED_BODY()
+
+	FAnimNextSchemaAction_AddManifestNode() = default;
+
+	FAnimNextSchemaAction_AddManifestNode(const FAnimNextAssetRegistryManifestNode& InManifestNodeData, const FText& InKeywords = FText::GetEmpty());
+
+	// FEdGraphSchemaAction Interface
+	virtual UEdGraphNode* PerformAction(UEdGraph* ParentGraph, TArray<UEdGraphPin*>& FromPins, const FVector2D Location, bool bSelectNewNode = true) override;
+	virtual UEdGraphNode* PerformAction(UEdGraph* ParentGraph, UEdGraphPin* FromPin, const FVector2D Location, bool bSelectNewNode = true) override { return nullptr; }
+
+	TSoftObjectPtr<URigVMGraph> ModelGraph;
+	FString NodeName;
+};
diff --git a/Engine/Plugins/Experimental/AnimNextAnimGraph/Source/AnimNextAnimGraphUncookedOnly/Internal/Graph/AnimNextAnimationGraphSchema.h b/Engine/Plugins/Experimental/AnimNextAnimGraph/Source/AnimNextAnimGraphUncookedOnly/Internal/Graph/AnimNextAnimationGraphSchema.h
--- a/Engine/Plugins/Experimental/AnimNextAnimGraph/Source/AnimNextAnimGraphUncookedOnly/Internal/Graph/AnimNextAnimationGraphSchema.h
+++ b/Engine/Plugins/Experimental/AnimNextAnimGraph/Source/AnimNextAnimGraphUncookedOnly/Internal/Graph/AnimNextAnimationGraphSchema.h
@@ -1,4 +1,4 @@
-﻿// Copyright Epic Games, Inc. All Rights Reserved.
+// Copyright Epic Games, Inc. All Rights Reserved.
 
 #pragma once
 
@@ -11,5 +11,9 @@ class UAnimNextAnimationGraphSchema : public UAnimNextRigVMAssetSchema
 {
 	GENERATED_BODY()
 
+public:
 	// Schema reserved for future filtering of units etc.
+
+	// Returns true if a node supports setting a node title
+	virtual bool CanSetNodeTitle(URigVMController* InController, const URigVMNode* InNode) const override;
 };
diff --git a/Engine/Plugins/Experimental/AnimNextAnimGraph/Source/AnimNextAnimGraphUncookedOnly/Internal/Graph/AnimNextAnimationGraph_EditorData.h b/Engine/Plugins/Experimental/AnimNextAnimGraph/Source/AnimNextAnimGraphUncookedOnly/Internal/Graph/AnimNextAnimationGraph_EditorData.h
--- a/Engine/Plugins/Experimental/AnimNextAnimGraph/Source/AnimNextAnimGraphUncookedOnly/Internal/Graph/AnimNextAnimationGraph_EditorData.h
+++ b/Engine/Plugins/Experimental/AnimNextAnimGraph/Source/AnimNextAnimGraphUncookedOnly/Internal/Graph/AnimNextAnimationGraph_EditorData.h
@@ -1,4 +1,4 @@
-﻿// Copyright Epic Games, Inc. All Rights Reserved.
+// Copyright Epic Games, Inc. All Rights Reserved.
 
 #pragma once
 
@@ -16,6 +16,7 @@ class FAnimationAnimNextRuntimeTest_GraphExecuteLatent;
 namespace UE::AnimNext::UncookedOnly
 {
 	struct FUtils;
+	struct FAnimGraphUtils;
 }
 
 namespace UE::AnimNext::Editor
@@ -49,12 +50,14 @@ class ANIMNEXTANIMGRAPHUNCOOKEDONLY_API UAnimNextAnimationGraph_EditorData : pub
 	friend struct UE::AnimNext::Editor::FUtils;
 	friend class UE::AnimNext::Editor::FModuleEditor;
 	friend class UE::AnimNext::Editor::SAnimNextGraphView;
+	friend struct UE::AnimNext::UncookedOnly::FAnimGraphUtils;
 	friend struct FAnimNextGraphSchemaAction_RigUnit;
 	friend struct FAnimNextGraphSchemaAction_DispatchFactory;
 	friend class FAnimationAnimNextEditorTest_GraphAddTrait;
 	friend class FAnimationAnimNextEditorTest_GraphTraitOperations;
 	friend class FAnimationAnimNextRuntimeTest_GraphExecute;
 	friend class FAnimationAnimNextRuntimeTest_GraphExecuteLatent;
+	friend class FAnimationAnimNextEditorTest_GraphManifest;
 
 public:
 	/** Adds an animation graph to this asset */
@@ -72,4 +75,7 @@ protected:
 	virtual TSubclassOf<UAssetUserData> GetAssetUserDataClass() const override;
 	virtual void InitializeAssetUserData() override;
 	virtual void GetProgrammaticGraphs(const FRigVMCompileSettings& InSettings, TArray<URigVMGraph*>& OutGraphs) override;
+
+	// Super overrides
+	virtual void GetAnimNextAssetRegistryTags(FAssetRegistryTagsContext& Context, FAnimNextAssetRegistryExports& OutExports) const override;
 };
diff --git a/Engine/Plugins/Experimental/AnimNextAnimGraph/Source/AnimNextAnimGraphUncookedOnly/Private/AnimGraphUncookedOnlyUtils.cpp b/Engine/Plugins/Experimental/AnimNextAnimGraph/Source/AnimNextAnimGraphUncookedOnly/Private/AnimGraphUncookedOnlyUtils.cpp
--- a/Engine/Plugins/Experimental/AnimNextAnimGraph/Source/AnimNextAnimGraphUncookedOnly/Private/AnimGraphUncookedOnlyUtils.cpp
+++ b/Engine/Plugins/Experimental/AnimNextAnimGraph/Source/AnimNextAnimGraphUncookedOnly/Private/AnimGraphUncookedOnlyUtils.cpp
@@ -1,16 +1,22 @@
-﻿// Copyright Epic Games, Inc. All Rights Reserved.
+// Copyright Epic Games, Inc. All Rights Reserved.
 
 #include "AnimGraphUncookedOnlyUtils.h"
 
 #include "AnimNextEdGraphNode.h"
 #include "Graph/RigUnit_AnimNextGraphRoot.h"
 #include "Graph/RigUnit_AnimNextTraitStack.h"
-#include "RigVMModel/RigVMController.h"	
+#include "AnimNextController.h"
+#include "TraitCore/TraitRegistry.h"
+#include "UncookedOnlyUtils.h"
+#include "AnimNextRigVMAssetEditorData.h"
+#include "Graph/AnimNextAnimationGraph_EditorData.h"
+#include "AssetRegistry/IAssetRegistry.h"
+#include "AnimNextTraitStackUnitNode.h"
 
 namespace UE::AnimNext::UncookedOnly
 {
 
-bool FAnimGraphUtils::IsTraitStackNode(URigVMNode* InModelNode)
+bool FAnimGraphUtils::IsTraitStackNode(const URigVMNode* InModelNode)
 {
 	if (const URigVMUnitNode* VMNode = Cast<URigVMUnitNode>(InModelNode))
 	{
@@ -48,4 +54,86 @@ void FAnimGraphUtils::SetupAnimGraph(const FName EntryName, URigVMController* In
 	InController->SetPinDefaultValue(EntryPointPin->GetPinPath(), EntryName.ToString(), true, bSetupUndoRedo);
 }
 
+void FAnimGraphUtils::GetAssetManifestNodesRegistryExports(const UAnimNextRigVMAssetEditorData* InEditorData, FAnimNextAssetRegistryExports& OutExports)
+{
+	if (const UAnimNextAnimationGraph_EditorData* EditorData = Cast<UAnimNextAnimationGraph_EditorData>(InEditorData))
+	{
+		TArray<UAnimNextEdGraphNode*> AllNodes;
+		EditorData->GetAllNodesOfClass(AllNodes);
+
+		OutExports.ManifestNodes.Reset();
+		OutExports.ManifestNodes.Reserve(AllNodes.Num());
+
+		for (const UAnimNextEdGraphNode* EdNode : AllNodes)
+		{
+			if (const URigVMNode* ModelNode = EdNode->GetModelNode())
+			{
+				if (UE::AnimNext::UncookedOnly::FAnimGraphUtils::IsExposedToManifest(ModelNode))
+				{
+					const UPackage* Package = EditorData->GetOuter()->GetPackage();
+					FString PackageName;
+					Package->GetName().Split(FString(TEXT("/")), nullptr, &PackageName, ESearchCase::IgnoreCase, ESearchDir::FromEnd);
+
+					OutExports.ManifestNodes.Add(FAnimNextAssetRegistryManifestNode(
+						ModelNode->GetGraph(),
+						*ModelNode->GetName(),
+						PackageName, // for now we set the package as category
+						EdNode->GetNodeTitle(ENodeTitleType::MenuTitle).ToString(),
+						ModelNode->GetToolTipText().ToString()));
+				}
+			}
+		}
+	}
+}
+
+bool FAnimGraphUtils::GetExportedManifestNodesFromAssetRegistry(TArray<FAnimNextAssetRegistryExports>& OutExports)
+{
+	TArray<FAssetData> AssetData;
+	IAssetRegistry::GetChecked().GetAssetsByTags({ UE::AnimNext::ExportsAnimNextAssetRegistryTag }, AssetData);
+
+	for (const FAssetData& Asset : AssetData)
+	{
+		const FString TagValue = Asset.GetTagValueRef<FString>(UE::AnimNext::ExportsAnimNextAssetRegistryTag);
+		FAnimNextAssetRegistryExports AssetExports;
+		if (FAnimNextAssetRegistryExports::StaticStruct()->ImportText(*TagValue, &AssetExports, nullptr, PPF_None, nullptr, FAnimNextAssetRegistryExports::StaticStruct()->GetName()) != nullptr)
+		{
+			if (AssetExports.ManifestNodes.Num() > 0)
+			{
+				OutExports.Add(MoveTemp(AssetExports));
+			}
+		}
+	}
+
+	return OutExports.Num() > 0;
 }
+
+bool FAnimGraphUtils::IsExposedToManifest(const URigVMNode* InModelNode)
+{
+	bool bIsExposedToManifest = false;
+
+	if (InModelNode != nullptr)
+	{
+		const UAnimNextTraitStackUnitNode* AnimNextUnitNode = Cast<UAnimNextTraitStackUnitNode>(InModelNode);
+		if (AnimNextUnitNode != nullptr)
+		{
+			bIsExposedToManifest = AnimNextUnitNode->IsExposedToManifest();
+		}
+	}
+
+	return bIsExposedToManifest;
+}
+
+bool FAnimGraphUtils::RequestVMAutoRecompile(UAnimNextRigVMAssetEditorData* EditorData)
+{
+	if (EditorData)
+	{
+		if (IRigVMClientHost* ClientHost = CastChecked<IRigVMClientHost>(EditorData))
+		{
+			ClientHost->RequestAutoVMRecompilation();
+			return true;
+		}
+	}
+	return false;
+}
+
+} // end namespace
diff --git a/Engine/Plugins/Experimental/AnimNextAnimGraph/Source/AnimNextAnimGraphUncookedOnly/Private/AnimNextController.cpp b/Engine/Plugins/Experimental/AnimNextAnimGraph/Source/AnimNextAnimGraphUncookedOnly/Private/AnimNextController.cpp
--- a/Engine/Plugins/Experimental/AnimNextAnimGraph/Source/AnimNextAnimGraphUncookedOnly/Private/AnimNextController.cpp
+++ b/Engine/Plugins/Experimental/AnimNextAnimGraph/Source/AnimNextAnimGraphUncookedOnly/Private/AnimNextController.cpp
@@ -2,7 +2,10 @@
 
 #include "AnimNextController.h"
 
-#include "AnimNextUnitNode.h"
+#include "AnimNextControllerActions.h"
+#include "RigVMPythonUtils.h"
+#include "AnimGraphUncookedOnlyUtils.h"
+#include "AnimNextTraitStackUnitNode.h"
 #include "Graph/RigDecorator_AnimNextCppTrait.h"
 #include "Param/ParamType.h"
 #include "RigVMModel/RigVMControllerActions.h"
@@ -238,3 +241,167 @@ URigVMUnitNode* UAnimNextController::AddUnitNodeWithPins(UScriptStruct* InScript
 
 	return Node;
 }
+
+bool UAnimNextController::AddNodeToManifest(URigVMNode* InModelNode, bool bSetupUndoRedo, bool bPrintPythonCommand)
+{
+	if (!IsValidGraph())
+	{
+		return false;
+	}
+
+	if (!IsTransacting() && !IsGraphEditable())
+	{
+		return false;
+	}
+
+	if (!ensure(InModelNode != nullptr))
+	{
+		return false;
+	}
+
+	UAnimNextTraitStackUnitNode* AnimNextUnitNode = Cast<UAnimNextTraitStackUnitNode>(InModelNode);
+	if (AnimNextUnitNode == nullptr)
+	{
+		return false;
+	}
+
+	if (!ensure(UE::AnimNext::UncookedOnly::FAnimGraphUtils::IsExposedToManifest(InModelNode) == false))
+	{
+		return false;
+	}
+
+	const URigVMGraph* Graph = GetGraph();
+	check(Graph);
+
+	// Avoid multiple VM recompilations for internal operations
+	FRigVMControllerCompileBracketScope CompileScope(this);
+
+	FAnimNextManifestAction Action;
+	if (bSetupUndoRedo)
+	{
+		Action = FAnimNextManifestAction(this, InModelNode, true);
+		Action.SetTitle(FString::Printf(TEXT("Add Node to Manifest")));
+		GetActionStack()->BeginAction(Action);
+	}
+
+	AnimNextUnitNode->SetManifestState(true);
+
+	if (bSetupUndoRedo)
+	{
+		GetActionStack()->EndAction(Action);
+	}
+
+	if (bPrintPythonCommand)
+	{
+		const FString GraphName = GetSchema()->GetSanitizedGraphName(GetGraph()->GetGraphName());
+		const FString NodePath = GetSchema()->GetSanitizedPinPath(InModelNode->GetNodePath());
+
+		RigVMPythonUtils::Print(GetSchema()->GetGraphOuterName(GetGraph()),
+			FString::Printf(TEXT("blueprint.get_controller_by_name('%s').add_node_to_manifest_by_name('%s')"),
+				*GraphName,
+				*NodePath));
+	}
+
+	return true;
+}
+
+bool UAnimNextController::AddNodeToManifestByName(const FName& InNodeName, bool bSetupUndoRedo, bool bPrintPythonCommand)
+{
+	if (!IsValidGraph())
+	{
+		return false;
+	}
+
+	if (!IsTransacting() && !IsGraphEditable())
+	{
+		return false;
+	}
+
+	URigVMGraph* Graph = GetGraph();
+	check(Graph);
+
+	URigVMNode* Node = Graph->FindNodeByName(InNodeName);
+	return AddNodeToManifest(Node, bSetupUndoRedo, bPrintPythonCommand);
+}
+
+bool UAnimNextController::RemoveNodeFromManifest(URigVMNode* InModelNode, bool bSetupUndoRedo, bool bPrintPythonCommand)
+{
+	if (!IsValidGraph())
+	{
+		return false;
+	}
+
+	if (!IsTransacting() && !IsGraphEditable())
+	{
+		return false;
+	}
+
+	if (!ensure(InModelNode != nullptr))
+	{
+		return false;
+	}
+
+	UAnimNextTraitStackUnitNode* AnimNextUnitNode = Cast<UAnimNextTraitStackUnitNode>(InModelNode);
+	if (AnimNextUnitNode == nullptr)
+	{
+		return false;
+	}
+
+	if (!ensure(UE::AnimNext::UncookedOnly::FAnimGraphUtils::IsExposedToManifest(InModelNode)))
+	{
+		return false;
+	}
+
+	const URigVMGraph* Graph = GetGraph();
+	check(Graph);
+
+	// Avoid multiple VM recompilations for internal operations
+	FRigVMControllerCompileBracketScope CompileScope(this);
+
+	FAnimNextManifestAction Action;
+	if (bSetupUndoRedo)
+	{
+		Action = FAnimNextManifestAction(this, AnimNextUnitNode, false);
+		Action.SetTitle(FString::Printf(TEXT("Remove Node from Manifest")));
+		GetActionStack()->BeginAction(Action);
+	}
+
+	AnimNextUnitNode->SetManifestState(false);
+
+	if (bSetupUndoRedo)
+	{
+		GetActionStack()->EndAction(Action);
+	}
+
+	if (bPrintPythonCommand)
+	{
+		const FString GraphName = GetSchema()->GetSanitizedGraphName(GetGraph()->GetGraphName());
+		const FString NodePath = GetSchema()->GetSanitizedPinPath(InModelNode->GetNodePath());
+
+		RigVMPythonUtils::Print(GetSchema()->GetGraphOuterName(GetGraph()),
+			FString::Printf(TEXT("blueprint.get_controller_by_name('%s').remove_node_from_manifest_by_name('%s')"),
+				*GraphName,
+				*NodePath));
+	}
+
+	return true;
+}
+
+bool UAnimNextController::RemoveNodeFromManifestByName(const FName& InNodeName, bool bSetupUndoRedo, bool bPrintPythonCommand)
+{
+	if (!IsValidGraph())
+	{
+		return false;
+	}
+
+	if (!IsTransacting() && !IsGraphEditable())
+	{
+		return false;
+	}
+
+	URigVMGraph* Graph = GetGraph();
+	check(Graph);
+
+	URigVMNode* Node = Graph->FindNodeByName(InNodeName);
+	return RemoveNodeFromManifest(Node, bSetupUndoRedo, bPrintPythonCommand);
+}
diff --git a/Engine/Plugins/Experimental/AnimNextAnimGraph/Source/AnimNextAnimGraphUncookedOnly/Private/AnimNextControllerActions.cpp b/Engine/Plugins/Experimental/AnimNextAnimGraph/Source/AnimNextAnimGraphUncookedOnly/Private/AnimNextControllerActions.cpp
new file mode 100644
--- /dev/null
+++ b/Engine/Plugins/Experimental/AnimNextAnimGraph/Source/AnimNextAnimGraphUncookedOnly/Private/AnimNextControllerActions.cpp
@@ -0,0 +1,74 @@
+#include "AnimNextControllerActions.h"
+#include "AnimGraphUncookedOnlyUtils.h"
+
+FAnimNextManifestAction::FAnimNextManifestAction()
+	: FAnimNextBaseAction(nullptr)
+{
+}
+
+FAnimNextManifestAction::FAnimNextManifestAction(UAnimNextController* InController, URigVMNode* InNode, bool bInNewIncludeInManifestState)
+	: FAnimNextBaseAction(InController)
+	, NodePath(InNode->GetNodePath())
+	, bOldManifestState(UE::AnimNext::UncookedOnly::FAnimGraphUtils::IsExposedToManifest(InNode))
+	, bIncludeInManifestState(bInNewIncludeInManifestState)
+{
+}
+
+bool FAnimNextManifestAction::Merge(const FRigVMBaseAction* Other)
+{
+	if (!FRigVMBaseAction::Merge(Other))
+	{
+		return false;
+	}
+
+	const FAnimNextManifestAction* Action = (const FAnimNextManifestAction*)Other;
+	if (NodePath != Action->NodePath)
+	{
+		return false;
+	}
+
+	bIncludeInManifestState = Action->bIncludeInManifestState;
+	return true;
+}
+
+bool FAnimNextManifestAction::Undo()
+{
+	if (!FRigVMBaseAction::Undo())
+	{
+		return false;
+	}
+
+	if (bOldManifestState)
+	{
+		return GetAnimNextController()->AddNodeToManifestByName(*NodePath, false, false);
+	}
+	else
+	{
+		return GetAnimNextController()->RemoveNodeFromManifestByName(*NodePath, false, false);
+	}
+}
+
+bool FAnimNextManifestAction::Redo()
+{
+	if (!CanUndoRedo())
+	{
+		return false;
+	}
+
+	if (bIncludeInManifestState)
+	{
+		if (!GetAnimNextController()->AddNodeToManifestByName(*NodePath, false, false))
+		{
+			return false;
+		}
+	}
+	else
+	{
+		if (!GetAnimNextController()->RemoveNodeFromManifestByName(*NodePath, false, false))
+		{
+			return false;
+		}
+	}
+
+	return FRigVMBaseAction::Redo();
+}
diff --git a/Engine/Plugins/Experimental/AnimNextAnimGraph/Source/AnimNextAnimGraphUncookedOnly/Private/AnimNextControllerActions.h b/Engine/Plugins/Experimental/AnimNextAnimGraph/Source/AnimNextAnimGraphUncookedOnly/Private/AnimNextControllerActions.h
new file mode 100644
--- /dev/null
+++ b/Engine/Plugins/Experimental/AnimNextAnimGraph/Source/AnimNextAnimGraphUncookedOnly/Private/AnimNextControllerActions.h
@@ -0,0 +1,69 @@
+// Copyright Epic Games, Inc. All Rights Reserved.
+
+#pragma once
+
+#include "CoreMinimal.h"
+#include "RigVMModel/RigVMControllerActions.h"
+#include "AnimNextController.h"
+// --- ---
+#include "AnimNextControllerActions.generated.h"
+
+
+/**
+ * The base action is the base struct for all actions, and provides
+ * access to sub actions, merge functionality as well as undo and redo
+ * base implementations.
+ */
+USTRUCT()
+struct FAnimNextBaseAction : public FRigVMBaseAction
+{
+	GENERATED_BODY()
+
+public:
+	FAnimNextBaseAction() = default;
+
+	explicit FAnimNextBaseAction(URigVMController* InController)
+		: FRigVMBaseAction(InController)
+	{
+		if(InController)
+		{
+			ControllerPath = FSoftObjectPath(InController);
+		}
+	}
+
+
+	// Returns the controller of this action
+	inline UAnimNextController* GetAnimNextController() const
+	{
+		return CastChecked<UAnimNextController>(GetController());
+	}
+};
+
+/**
+ * An action to add or remove a node from the node manifest.
+ */
+USTRUCT()
+struct FAnimNextManifestAction : public FAnimNextBaseAction
+{
+	GENERATED_BODY()
+
+public:
+
+	FAnimNextManifestAction();
+	FAnimNextManifestAction(UAnimNextController* InController, URigVMNode* InNode, bool bInNewIncludeInManifestState);
+
+	virtual ~FAnimNextManifestAction() = default;
+	virtual UScriptStruct* GetScriptStruct() const override { return FAnimNextManifestAction::StaticStruct(); }
+	virtual bool Merge(const FRigVMBaseAction* Other);
+	virtual bool Undo() override;
+	virtual bool Redo() override;
+
+	UPROPERTY()
+	FString NodePath;
+
+	UPROPERTY()
+	bool bOldManifestState = false;
+
+	UPROPERTY()
+	bool bIncludeInManifestState = false;
+};
diff --git a/Engine/Plugins/Experimental/AnimNextAnimGraph/Source/AnimNextAnimGraphUncookedOnly/Private/Graph/AnimNextAnimationGraphSchema.cpp b/Engine/Plugins/Experimental/AnimNextAnimGraph/Source/AnimNextAnimGraphUncookedOnly/Private/Graph/AnimNextAnimationGraphSchema.cpp
--- a/Engine/Plugins/Experimental/AnimNextAnimGraph/Source/AnimNextAnimGraphUncookedOnly/Private/Graph/AnimNextAnimationGraphSchema.cpp
+++ b/Engine/Plugins/Experimental/AnimNextAnimGraph/Source/AnimNextAnimGraphUncookedOnly/Private/Graph/AnimNextAnimationGraphSchema.cpp
@@ -1,4 +1,16 @@
-﻿// Copyright Epic Games, Inc. All Rights Reserved.
+// Copyright Epic Games, Inc. All Rights Reserved.
 
 #include "Graph/AnimNextAnimationGraphSchema.h"
+#include "AnimGraphUncookedOnlyUtils.h"
 
+bool UAnimNextAnimationGraphSchema::CanSetNodeTitle(URigVMController* InController, const URigVMNode* InNode) const
+{
+	if (InNode != nullptr)
+	{
+		if (UE::AnimNext::UncookedOnly::FAnimGraphUtils::IsTraitStackNode(InNode))
+		{
+			return true;
+		}
+	}
+	return false;
+}
diff --git a/Engine/Plugins/Experimental/AnimNextAnimGraph/Source/AnimNextAnimGraphUncookedOnly/Private/Graph/AnimNextAnimationGraph_EditorData.cpp b/Engine/Plugins/Experimental/AnimNextAnimGraph/Source/AnimNextAnimGraphUncookedOnly/Private/Graph/AnimNextAnimationGraph_EditorData.cpp
--- a/Engine/Plugins/Experimental/AnimNextAnimGraph/Source/AnimNextAnimGraphUncookedOnly/Private/Graph/AnimNextAnimationGraph_EditorData.cpp
+++ b/Engine/Plugins/Experimental/AnimNextAnimGraph/Source/AnimNextAnimGraphUncookedOnly/Private/Graph/AnimNextAnimationGraph_EditorData.cpp
@@ -26,6 +26,7 @@
 #include "UObject/AssetRegistryTagsContext.h"
 #include "UObject/LinkerLoad.h"
 #include "Traits/CallFunction.h"
+#include "Entries/AnimNextRigVMAssetEntry.h"
 
 namespace UE::AnimNext::UncookedOnly::Private
 {
@@ -188,7 +189,7 @@ namespace UE::AnimNext::UncookedOnly::Private
 					return FString();
 				}
 
-				// A regular property pin
+				// A regular property pin or hidden pin
 				return Pin->GetDefaultValue();
 			}
 		}
@@ -455,6 +456,7 @@ namespace UE::AnimNext::UncookedOnly::Private
 void UAnimNextAnimationGraph_EditorData::RecompileVM()
 {
 	using namespace UE::AnimNext::UncookedOnly;
+	using namespace UE::AnimNext::UncookedOnly::Private;
 
 	if (bIsCompiling)
 	{
@@ -547,25 +549,25 @@ void UAnimNextAnimationGraph_EditorData::RecompileVM()
 
 			FRigVMPinInfoArray LatentPins;
 			TMap<FName, URigVMPin*> LatentPinMapping;
-			TArray<Private::FTraitGraph> TraitGraphs;
+			TArray<FTraitGraph> TraitGraphs;
 
 			// Build entry points and extract their required latent pins
 			for(const URigVMGraph* AnimGraph : AnimGraphs)
 			{
 				// Gather our trait stacks
-				Private::FTraitGraph& TraitGraph = TraitGraphs.Add_GetRef(Private::CollectGraphInfo(AnimationGraph, AnimGraph, TempController->GetControllerForGraph(AnimGraph)));
+				FTraitGraph& TraitGraph = TraitGraphs.Add_GetRef(CollectGraphInfo(AnimationGraph, AnimGraph, TempController->GetControllerForGraph(AnimGraph)));
 				check(!TraitGraph.TraitStackNodes.IsEmpty());
 
 				FAnimNextGraphEntryPoint& EntryPoint = AnimationGraph->EntryPoints.AddDefaulted_GetRef();
 				EntryPoint.EntryPointName = TraitGraph.EntryPoint;
 
 				// Extract latent pins for this graph
-				Private::CollectLatentPins(TraitGraph.TraitStackNodes, LatentPins, LatentPinMapping);
+				CollectLatentPins(TraitGraph.TraitStackNodes, LatentPins, LatentPinMapping);
 
 				// Iterate over every trait stack and register our node templates
-				for (Private::FTraitStackMapping& NodeMapping : TraitGraph.TraitStackNodes)
+				for (FTraitStackMapping& NodeMapping : TraitGraph.TraitStackNodes)
 				{
-					NodeMapping.TraitStackNodeHandle = Private::RegisterTraitNodeTemplate(TraitWriter, NodeMapping.DecoratorStackNode);
+					NodeMapping.TraitStackNodeHandle = RegisterTraitNodeTemplate(TraitWriter, NodeMapping.DecoratorStackNode);
 				}
 
 				// Find our root node handle, if we have any stack nodes, the first one is our root stack
@@ -582,7 +584,7 @@ void UAnimNextAnimationGraph_EditorData::RecompileVM()
 			}
 
 			// Remove our old root nodes
-			for (Private::FTraitGraph& TraitGraph : TraitGraphs)
+			for (FTraitGraph& TraitGraph : TraitGraphs)
 			{
 				URigVMController* GraphController = TempController->GetControllerForGraph(TraitGraph.RootNode->GetGraph());
 				GraphController->RemoveNode(TraitGraph.RootNode, false, false);
@@ -591,7 +593,7 @@ void UAnimNextAnimationGraph_EditorData::RecompileVM()
 			if(LatentPins.Num() > 0)
 			{
 				// We need a unique method name to match our unique argument list
-				AnimationGraph->ExecuteDefinition = Private::GetGraphEvaluatorExecuteMethod(LatentPins);
+				AnimationGraph->ExecuteDefinition = GetGraphEvaluatorExecuteMethod(LatentPins);
 
 				// Add our runtime shim root node
 				URigVMUnitNode* TempShimRootNode = TempController->AddUnitNode(FRigUnit_AnimNextShimRoot::StaticStruct(), FRigUnit_AnimNextShimRoot::EventName, FVector2D::ZeroVector, FString(), false);
@@ -616,11 +618,11 @@ void UAnimNextAnimationGraph_EditorData::RecompileVM()
 			// Write our node shared data
 			TraitWriter.BeginNodeWriting();
 
-			for(Private::FTraitGraph& TraitGraph : TraitGraphs)
+			for(FTraitGraph& TraitGraph : TraitGraphs)
 			{
-				for (const Private::FTraitStackMapping& NodeMapping : TraitGraph.TraitStackNodes)
+				for (const FTraitStackMapping& NodeMapping : TraitGraph.TraitStackNodes)
 				{
-					Private::WriteTraitProperties(TraitWriter, NodeMapping, TraitGraph.TraitStackNodes);
+					WriteTraitProperties(TraitWriter, NodeMapping, TraitGraph.TraitStackNodes);
 				}
 			}
 
@@ -943,4 +945,9 @@ void UAnimNextAnimationGraph_EditorData::GetProgrammaticGraphs(const FRigVMCompi
 			OutGraphs.Add(WrapperGraph);
 		}
 	}
-}
\ No newline at end of file
+}
+
+void UAnimNextAnimationGraph_EditorData::GetAnimNextAssetRegistryTags(FAssetRegistryTagsContext& Context, FAnimNextAssetRegistryExports& OutExports) const
+{
+	UE::AnimNext::UncookedOnly::FAnimGraphUtils::GetAssetManifestNodesRegistryExports(this, OutExports);
+}
diff --git a/Engine/Plugins/Experimental/Workspace/Source/WorkspaceEditor/Private/WorkspaceEditor.cpp b/Engine/Plugins/Experimental/Workspace/Source/WorkspaceEditor/Private/WorkspaceEditor.cpp
--- a/Engine/Plugins/Experimental/Workspace/Source/WorkspaceEditor/Private/WorkspaceEditor.cpp
+++ b/Engine/Plugins/Experimental/Workspace/Source/WorkspaceEditor/Private/WorkspaceEditor.cpp
@@ -331,6 +331,17 @@ void FWorkspaceEditor::OpenObjects(TConstArrayView<UObject*> InObjects)
 	}
 }
 
+void FWorkspaceEditor::GetAssets(TArray<FAssetData>& OutAssets) const
+{
+	FWorkspaceAssetRegistryExports Exports;
+	FWorkspaceEditorModule::GetExportedAssetsForWorkspace(FAssetData(Workspace), Exports);
+	IAssetRegistry& AssetRegistry = FModuleManager::LoadModuleChecked<FAssetRegistryModule>("AssetRegistry").Get();
+	for (const FWorkspaceAssetRegistryExportEntry& Entry : Exports.Assets)
+	{
+		OutAssets.Add(AssetRegistry.GetAssetByObjectPath(Entry.Asset));
+	}
+}
+
 void FWorkspaceEditor::CloseObjects(TConstArrayView<UObject*> InObjects)
 {
 	if(InObjects.Num() > 0)
@@ -456,6 +467,11 @@ TSharedPtr<IDetailsView> FWorkspaceEditor::GetDetailsView()
 	return DetailsView;
 }
 
+FString FWorkspaceEditor::GetPackageName() const
+{
+	return Workspace->GetPackage()->GetName();
+}
+
 void FWorkspaceEditor::BindCommands()
 {
 	const FWorkspaceAssetEditorCommands& Commands = FWorkspaceAssetEditorCommands::Get();
diff --git a/Engine/Plugins/Experimental/Workspace/Source/WorkspaceEditor/Private/WorkspaceEditor.h b/Engine/Plugins/Experimental/Workspace/Source/WorkspaceEditor/Private/WorkspaceEditor.h
--- a/Engine/Plugins/Experimental/Workspace/Source/WorkspaceEditor/Private/WorkspaceEditor.h
+++ b/Engine/Plugins/Experimental/Workspace/Source/WorkspaceEditor/Private/WorkspaceEditor.h
@@ -96,6 +96,7 @@ private:
 	// IWorkspaceEditor interface
 	virtual void OpenAssets(TConstArrayView<FAssetData> InAssets) override;
 	virtual void OpenObjects(TConstArrayView<UObject*> InObjects) override;
+	virtual void GetAssets(TArray<FAssetData>& OutAssets) const override;
 	virtual void CloseObjects(TConstArrayView<UObject*> InObjects) override;
 	virtual void SetDetailsObjects(const TArray<UObject*>& InObjects) override;
 	virtual void RefreshDetails() override;
@@ -107,6 +108,7 @@ private:
 	virtual const TObjectPtr<UObject> GetFocussedAssetOfClass(const TObjectPtr<UClass> AssetClass) const override;
 	virtual FOnFocussedAssetChanged& OnFocussedAssetChanged() override;
 	virtual TSharedPtr<IDetailsView> GetDetailsView() override;
+	virtual FString GetPackageName() const override;
 
 	void HandleOutlinerSelectionChanged(TConstArrayView<FWorkspaceOutlinerItemExport> InExports);
 
diff --git a/Engine/Plugins/Experimental/Workspace/Source/WorkspaceEditor/Public/IWorkspaceEditor.h b/Engine/Plugins/Experimental/Workspace/Source/WorkspaceEditor/Public/IWorkspaceEditor.h
--- a/Engine/Plugins/Experimental/Workspace/Source/WorkspaceEditor/Public/IWorkspaceEditor.h
+++ b/Engine/Plugins/Experimental/Workspace/Source/WorkspaceEditor/Public/IWorkspaceEditor.h
@@ -1,4 +1,4 @@
-﻿// Copyright Epic Games, Inc. All Rights Reserved.
+// Copyright Epic Games, Inc. All Rights Reserved.
 
 #pragma once
 
@@ -6,6 +6,7 @@
 
 class UWorkspaceAssetEditor;
 class UWorkspaceSchema;
+class UWorkspace;
 struct FWorkspaceOutlinerItemExport;
 
 namespace UE::Workspace
@@ -36,6 +37,9 @@ public:
 	// Open the supplied objects for editing within the workspace editor
 	virtual void OpenObjects(TConstArrayView<UObject*> InObjects) = 0;
 
+	// Get the current set of assets in this workspace editor 
+	virtual void GetAssets(TArray<FAssetData>& OutAssets) const = 0;
+
 	// Close the supplied objects if they are open for editing within the workspace editor
 	virtual void CloseObjects(TConstArrayView<UObject*> InObjects) = 0;
 
@@ -71,6 +75,9 @@ public:
 
 	// Retrieves the common DetailsView widget
 	virtual TSharedPtr<IDetailsView> GetDetailsView() = 0;
+
+	// Returns the name of the package where the workspace is located
+	virtual FString GetPackageName() const = 0;
 };
 
 }
\ No newline at end of file
diff --git a/Engine/Plugins/Runtime/RigVM/Source/RigVMDeveloper/Public/RigVMModel/RigVMControllerActions.h b/Engine/Plugins/Runtime/RigVM/Source/RigVMDeveloper/Public/RigVMModel/RigVMControllerActions.h
--- a/Engine/Plugins/Runtime/RigVM/Source/RigVMDeveloper/Public/RigVMModel/RigVMControllerActions.h
+++ b/Engine/Plugins/Runtime/RigVM/Source/RigVMDeveloper/Public/RigVMModel/RigVMControllerActions.h
@@ -72,7 +72,7 @@ struct FRigVMActionKey
 /**
  * The action wrapper is used to extract an action from a serialized key.
  */
-struct FRigVMActionWrapper
+struct RIGVMDEVELOPER_API FRigVMActionWrapper
 {
 public:
 	FRigVMActionWrapper(const FRigVMActionKey& Key);
@@ -112,7 +112,7 @@ struct FRigVMActionNodeContent
  * base implementations.
  */
 USTRUCT()
-struct FRigVMBaseAction
+struct RIGVMDEVELOPER_API FRigVMBaseAction
 {
 	GENERATED_BODY()
 
