Fix race causing delegate handle leak in AnimNext
Fix incorrect usage of weak obj ptr as a map key
Fix delegate being static when it should be a member and hide it as we anticipate it will go away

diff --git a/Engine/Plugins/Experimental/AnimNext/Source/AnimNext/Private/DataRegistry.cpp b/Engine/Plugins/Experimental/AnimNext/Source/AnimNext/Private/DataRegistry.cpp
--- a/Engine/Plugins/Experimental/AnimNext/Source/AnimNext/Private/DataRegistry.cpp
+++ b/Engine/Plugins/Experimental/AnimNext/Source/AnimNext/Private/DataRegistry.cpp
@@ -65,8 +65,8 @@ FDataRegistry* FDataRegistry::Get()
 
 		for (auto Iter = GAnimationDataRegistry->SkeletalMeshReferencePoses.CreateIterator(); Iter; ++Iter)
 		{
-			const TWeakObjectPtr<const USkeletalMeshComponent>& SkeletalMeshComponentPtr = Iter.Key();
-			if (SkeletalMeshComponentPtr.Get() == nullptr)
+			const TObjectKey<USkeletalMeshComponent>& SkeletalMeshComponentPtr = Iter.Key();
+			if (SkeletalMeshComponentPtr.ResolveObjectPtr() == nullptr)
 			{
 				Iter.RemoveCurrent();
 			}
@@ -83,11 +83,24 @@ FDataHandle FDataRegistry::RegisterReferencePose(USkeletalMeshComponent* Skeleta
 	FGenerationTools::GenerateReferencePose(SkeletalMeshComponent, SkeletalMeshComponent->GetSkeletalMeshAsset(), AnimationReferencePose);
 
 	// register even if it fails to generate (register an empty ref pose)
-	{
-		FDelegateHandle DelegateHandle = SkeletalMeshComponent->RegisterOnLODRequiredBonesUpdate(FOnLODRequiredBonesUpdate::CreateRaw(this, &FDataRegistry::OnLODRequiredBonesUpdate));
 
+	const TObjectKey<USkeletalMeshComponent> Key(SkeletalMeshComponent);
+	const uint32 KeyHash = GetTypeHash(Key);
+
+	{
 		FRWScopeLock Lock(SkeletalMeshReferencePosesLock, SLT_Write);
-		SkeletalMeshReferencePoses.Add(SkeletalMeshComponent, FReferencePoseData(Handle, DelegateHandle));
+
+		// First check if another thread grabbed the write lock before us to do the same work
+		if (const FReferencePoseData* ExistingPoseData = SkeletalMeshReferencePoses.FindByHash(KeyHash, Key))
+		{
+			// Another thread beat us to it, discard the work we did and re-use what is cached
+			return ExistingPoseData->AnimationDataHandle;
+		}
+
+		// Only register the delegate if we are the one to add to the map
+		FDelegateHandle DelegateHandle = SkeletalMeshComponent->RegisterOnLODRequiredBonesUpdate_Member(FOnLODRequiredBonesUpdate::CreateRaw(this, &FDataRegistry::OnLODRequiredBonesUpdate));
+
+		SkeletalMeshReferencePoses.AddByHash(KeyHash, Key, FReferencePoseData(Handle, DelegateHandle));
 	}
 
 	return Handle;
@@ -129,7 +142,7 @@ void FDataRegistry::RemoveReferencePose(USkeletalMeshComponent* SkeletalMeshComp
 		FReferencePoseData ReferencePoseData;
 		if (SkeletalMeshReferencePoses.RemoveAndCopyValue(SkeletalMeshComponent, ReferencePoseData))
 		{
-			SkeletalMeshComponent->UnregisterOnLODRequiredBonesUpdate(ReferencePoseData.DelegateHandle);
+			SkeletalMeshComponent->UnregisterOnLODRequiredBonesUpdate_Member(ReferencePoseData.DelegateHandle);
 		}
 	}
 }
@@ -199,12 +212,12 @@ void FDataRegistry::ReleaseReferencePoseData()
 
 	for (auto Iter = GAnimationDataRegistry->SkeletalMeshReferencePoses.CreateIterator(); Iter; ++Iter)
 	{
-		const TWeakObjectPtr<USkeletalMeshComponent>& SkeletalMeshComponentPtr = Iter.Key();
+		const TObjectKey<USkeletalMeshComponent>& SkeletalMeshComponentPtr = Iter.Key();
 		const FReferencePoseData& ReferencePoseData = Iter.Value();
 
-		if (USkeletalMeshComponent* SkeletalMeshComponent = SkeletalMeshComponentPtr.Get())
+		if (USkeletalMeshComponent* SkeletalMeshComponent = SkeletalMeshComponentPtr.ResolveObjectPtr())
 		{
-			SkeletalMeshComponent->UnregisterOnLODRequiredBonesUpdate(ReferencePoseData.DelegateHandle);
+			SkeletalMeshComponent->UnregisterOnLODRequiredBonesUpdate_Member(ReferencePoseData.DelegateHandle);
 		}
 	}
 
diff --git a/Engine/Plugins/Experimental/AnimNext/Source/AnimNext/Public/DataRegistry.h b/Engine/Plugins/Experimental/AnimNext/Source/AnimNext/Public/DataRegistry.h
--- a/Engine/Plugins/Experimental/AnimNext/Source/AnimNext/Public/DataRegistry.h
+++ b/Engine/Plugins/Experimental/AnimNext/Source/AnimNext/Public/DataRegistry.h
@@ -181,7 +181,7 @@ private:
 	mutable FRWLock StoredDataLock;
 
 	// Map holding reference poses for SkeletalMeshes
-	TMap <TWeakObjectPtr<USkeletalMeshComponent>, FReferencePoseData> SkeletalMeshReferencePoses;
+	TMap <TObjectKey<USkeletalMeshComponent>, FReferencePoseData> SkeletalMeshReferencePoses;
 	mutable FRWLock SkeletalMeshReferencePosesLock;
 
 	// Registers a type and sets the allocation block size
diff --git a/Engine/Source/Runtime/Engine/Classes/Components/SkeletalMeshComponent.h b/Engine/Source/Runtime/Engine/Classes/Components/SkeletalMeshComponent.h
--- a/Engine/Source/Runtime/Engine/Classes/Components/SkeletalMeshComponent.h
+++ b/Engine/Source/Runtime/Engine/Classes/Components/SkeletalMeshComponent.h
@@ -41,6 +41,11 @@ struct FConstraintProfileProperties;
 struct FNavigableGeometryExport;
 struct FCompactPose;
 
+namespace UE::AnimNext
+{
+	class FDataRegistry;	// Temporary
+}
+
 enum class EClothingTeleportMode : uint8;
 
 #ifndef WITH_CLOTH_COLLISION_DETECTION
@@ -2667,8 +2672,18 @@ public:
 	ENGINE_API virtual void UnregisterOnBoneTransformsFinalizedDelegate(const FDelegateHandle& DelegateHandle) override;
 
 	/** Register/Unregister for OnLODRequiredBonesUpdate callback */
+	UE_DEPRECATED(5.6, "These should not be used as they rely on a global delegate, ANY change in LOD from ANY skeletal mesh component triggers the callback")
 	ENGINE_API FDelegateHandle RegisterOnLODRequiredBonesUpdate(const FOnLODRequiredBonesUpdate& Delegate);
+
+	UE_DEPRECATED(5.6, "These should not be used as they rely on a global delegate, ANY change in LOD from ANY skeletal mesh component triggers the callback")
 	ENGINE_API void UnregisterOnLODRequiredBonesUpdate(const FDelegateHandle& DelegateHandle);
+
+private:
+	// NOTE: We anticipate that this will go away in the future, it is meant to be a temporary solution
+	/** Register/Unregister for OnLODRequiredBonesUpdate callback */
+	ENGINE_API FDelegateHandle RegisterOnLODRequiredBonesUpdate_Member(const FOnLODRequiredBonesUpdate& Delegate);
+	ENGINE_API void UnregisterOnLODRequiredBonesUpdate_Member(const FDelegateHandle& DelegateHandle);
+	friend class UE::AnimNext::FDataRegistry;
 	
 private:
 
@@ -2705,8 +2720,9 @@ private:
 	/** Multicaster fired when this component bone transforms are finalized */
 	FOnBoneTransformsFinalizedMultiCast OnBoneTransformsFinalizedMC;
 
-	/** Static Multicaster fired when this component finalizes the regeneration of the required bones list for the current LOD*/
-	static ENGINE_API FOnLODRequiredBonesUpdateMulticast OnLODRequiredBonesUpdate;
+	/** Multicaster fired when this component finalizes the regeneration of the required bones list for the current LOD */
+	static ENGINE_API FOnLODRequiredBonesUpdateMulticast OnLODRequiredBonesUpdate_Static;
+	FOnLODRequiredBonesUpdateMulticast OnLODRequiredBonesUpdate_Member;
 
 	/** Mark current anim UID version to up-to-date. Called when it's recalculated */
 	ENGINE_API void MarkRequiredCurveUpToDate();
diff --git a/Engine/Source/Runtime/Engine/Private/Components/SkeletalMeshComponent.cpp b/Engine/Source/Runtime/Engine/Private/Components/SkeletalMeshComponent.cpp
--- a/Engine/Source/Runtime/Engine/Private/Components/SkeletalMeshComponent.cpp
+++ b/Engine/Source/Runtime/Engine/Private/Components/SkeletalMeshComponent.cpp
@@ -117,7 +117,7 @@ FAutoConsoleTaskPriority CPrio_ParallelAnimationEvaluationTask(
 	);
 
 /** Static Multicaster fired when SkeletalMeshComponent finalizes the regeneration of the required bones list for the current LOD*/
-/*static*/ FOnLODRequiredBonesUpdateMulticast USkeletalMeshComponent::OnLODRequiredBonesUpdate;
+/*static*/ FOnLODRequiredBonesUpdateMulticast USkeletalMeshComponent::OnLODRequiredBonesUpdate_Static;
 
 #if !UE_BUILD_SHIPPING
 CSV_DEFINE_CATEGORY(AnimationParallelEvaluation, true);
@@ -2091,7 +2091,8 @@ void USkeletalMeshComponent::RecalcRequiredBones(int32 LODIndex)
 	// Invalidate cached bones.
 	ClearCachedAnimProperties();
 
-	OnLODRequiredBonesUpdate.Broadcast(this, LODIndex, RequiredBones);
+	OnLODRequiredBonesUpdate_Static.Broadcast(this, LODIndex, RequiredBones);
+	OnLODRequiredBonesUpdate_Member.Broadcast(this, LODIndex, RequiredBones);
 }
 
 void USkeletalMeshComponent::MarkRequiredCurveUpToDate()
@@ -4631,12 +4632,22 @@ void USkeletalMeshComponent::UnregisterOnBoneTransformsFinalizedDelegate(const F
 
 FDelegateHandle USkeletalMeshComponent::RegisterOnLODRequiredBonesUpdate(const FOnLODRequiredBonesUpdate& Delegate)
 {
-	return OnLODRequiredBonesUpdate.Add(Delegate);
+	return OnLODRequiredBonesUpdate_Static.Add(Delegate);
 }
 
 void USkeletalMeshComponent::UnregisterOnLODRequiredBonesUpdate(const FDelegateHandle& DelegateHandle)
 {
-	OnLODRequiredBonesUpdate.Remove(DelegateHandle);
+	OnLODRequiredBonesUpdate_Static.Remove(DelegateHandle);
+}
+
+FDelegateHandle USkeletalMeshComponent::RegisterOnLODRequiredBonesUpdate_Member(const FOnLODRequiredBonesUpdate& Delegate)
+{
+	return OnLODRequiredBonesUpdate_Member.Add(Delegate);
+}
+
+void USkeletalMeshComponent::UnregisterOnLODRequiredBonesUpdate_Member(const FDelegateHandle& DelegateHandle)
+{
+	OnLODRequiredBonesUpdate_Member.Remove(DelegateHandle);
 }
 
 bool USkeletalMeshComponent::MoveComponentImpl(const FVector& Delta, const FQuat& NewRotation, bool bSweep, FHitResult* OutHit /*= nullptr*/, EMoveComponentFlags MoveFlags /*= MOVECOMP_NoFlags*/, ETeleportType Teleport /*= ETeleportType::None*/)
