AnimNext: Fix function and interface bindings not working

diff --git a/Engine/Plugins/Experimental/AnimNext/Source/AnimNext/Internal/Module/AnimNextModuleContextData.h b/Engine/Plugins/Experimental/AnimNext/Source/AnimNext/Internal/Module/AnimNextModuleContextData.h
--- a/Engine/Plugins/Experimental/AnimNext/Source/AnimNext/Internal/Module/AnimNextModuleContextData.h
+++ b/Engine/Plugins/Experimental/AnimNext/Source/AnimNext/Internal/Module/AnimNextModuleContextData.h
@@ -28,6 +28,12 @@ struct FAnimNextModuleContextData
 	{
 	}
 
+	// Get the object that the module instance is bound to, if any
+	UObject* GetObject() const
+	{
+		return ModuleInstance ? ModuleInstance->Object : nullptr;
+	}
+
 	FAnimNextModuleInstance& GetModuleInstance() const
 	{
 		check(ModuleInstance != nullptr);
diff --git a/Engine/Plugins/Experimental/AnimNext/Source/AnimNext/Internal/Module/AnimNextModuleInstance.h b/Engine/Plugins/Experimental/AnimNext/Source/AnimNext/Internal/Module/AnimNextModuleInstance.h
--- a/Engine/Plugins/Experimental/AnimNext/Source/AnimNext/Internal/Module/AnimNextModuleInstance.h
+++ b/Engine/Plugins/Experimental/AnimNext/Source/AnimNext/Internal/Module/AnimNextModuleInstance.h
@@ -64,11 +64,11 @@ struct FAnimNextModuleInstance : public FAnimNextDataInterfaceInstance
 	GENERATED_BODY()
 
 	FAnimNextModuleInstance();
-	FAnimNextModuleInstance(
+	ANIMNEXT_API FAnimNextModuleInstance(
 		UAnimNextModule* InModule,
 		UObject* InObject,
 		EAnimNextModuleInitMethod InInitMethod);
-	~FAnimNextModuleInstance();
+	ANIMNEXT_API ~FAnimNextModuleInstance();
 
 	// Setup the entry
 	void Initialize();
@@ -118,7 +118,7 @@ struct FAnimNextModuleInstance : public FAnimNextDataInterfaceInstance
 
 #if UE_ENABLE_DEBUG_DRAWING
 	// Get the debug draw interface
-	ANIMNEXT_API FRigVMDrawInterface& GetDebugDrawInterface();
+	ANIMNEXT_API FRigVMDrawInterface* GetDebugDrawInterface();
 
 	// Set whether to show the module's debug drawing instructions in the viewport
 	ANIMNEXT_API void ShowDebugDrawing(bool bInShowDebugDrawing);
diff --git a/Engine/Plugins/Experimental/AnimNext/Source/AnimNext/Internal/Variables/RigVMDispatch_CallHoistedAccessorFunction.h b/Engine/Plugins/Experimental/AnimNext/Source/AnimNext/Internal/Variables/RigVMDispatch_CallHoistedAccessorFunction.h
--- a/Engine/Plugins/Experimental/AnimNext/Source/AnimNext/Internal/Variables/RigVMDispatch_CallHoistedAccessorFunction.h
+++ b/Engine/Plugins/Experimental/AnimNext/Source/AnimNext/Internal/Variables/RigVMDispatch_CallHoistedAccessorFunction.h
@@ -7,12 +7,12 @@
 #include "RigVMDispatch_CallHoistedAccessorFunction.generated.h"
 
 /** Synthetic dispatch injected by the compiler to get a value via a hoisted accessor UFunction, not user instantiated */
-USTRUCT(meta = (Hidden, DisplayName = "Call Hoisted Function", Category="Internal"))
-struct ANIMNEXT_API FRigVMDispatch_CallHoistedAccessorFunction : public FRigVMDispatchFactory
+USTRUCT(meta = (Hidden, Category="Internal"))
+struct ANIMNEXT_API FRigVMDispatch_CallHoistedAccessorFunctionBase : public FRigVMDispatchFactory
 {
 	GENERATED_BODY()
 
-	FRigVMDispatch_CallHoistedAccessorFunction();
+	FRigVMDispatch_CallHoistedAccessorFunctionBase();
 
 	static inline const FLazyName ObjectName = FLazyName("Object");
 	static inline const FLazyName FunctionName = FLazyName("Function");
@@ -25,10 +25,36 @@ private:
 	virtual const TArray<FRigVMTemplateArgumentInfo>& GetArgumentInfos() const override;
 	virtual FRigVMTemplateTypeMap OnNewArgumentType(const FName& InArgumentName, TRigVMTypeIndex InTypeIndex) const override;
 	virtual bool IsSingleton() const override { return true; }
+};
+
+/** Script function specialization of FRigVMDispatch_CallHoistedAccessorFunctionBase */
+USTRUCT(meta = (Hidden, DisplayName = "Call Hoisted Native Function"))
+struct ANIMNEXT_API FRigVMDispatch_CallHoistedAccessorFunctionNative : public FRigVMDispatch_CallHoistedAccessorFunctionBase
+{
+	GENERATED_BODY()
 
+	FRigVMDispatch_CallHoistedAccessorFunctionNative();
+
+private:
 	virtual FRigVMFunctionPtr GetDispatchFunctionImpl(const FRigVMTemplateTypeMap& InTypes) const override
 	{
-		return &FRigVMDispatch_CallHoistedAccessorFunction::Execute;
+		return &FRigVMDispatch_CallHoistedAccessorFunctionNative::Execute;
 	}
 	static void Execute(FRigVMExtendedExecuteContext& InContext, FRigVMMemoryHandleArray Handles, FRigVMPredicateBranchArray RigVMBranches);
 };
+
+/** Script function specialization of FRigVMDispatch_CallHoistedAccessorFunctionBase */
+USTRUCT(meta = (Hidden, DisplayName = "Call Hoisted Script Function"))
+struct ANIMNEXT_API FRigVMDispatch_CallHoistedAccessorFunctionScript : public FRigVMDispatch_CallHoistedAccessorFunctionBase
+{
+	GENERATED_BODY()
+
+	FRigVMDispatch_CallHoistedAccessorFunctionScript();
+
+private:
+	virtual FRigVMFunctionPtr GetDispatchFunctionImpl(const FRigVMTemplateTypeMap& InTypes) const override
+	{
+		return &FRigVMDispatch_CallHoistedAccessorFunctionScript::Execute;
+	}
+	static void Execute(FRigVMExtendedExecuteContext& InContext, FRigVMMemoryHandleArray Handles, FRigVMPredicateBranchArray RigVMBranches);
+};
\ No newline at end of file
diff --git a/Engine/Plugins/Experimental/AnimNext/Source/AnimNext/Internal/Variables/RigVMDispatch_CallObjectAccessorFunction.h b/Engine/Plugins/Experimental/AnimNext/Source/AnimNext/Internal/Variables/RigVMDispatch_CallObjectAccessorFunction.h
--- a/Engine/Plugins/Experimental/AnimNext/Source/AnimNext/Internal/Variables/RigVMDispatch_CallObjectAccessorFunction.h
+++ b/Engine/Plugins/Experimental/AnimNext/Source/AnimNext/Internal/Variables/RigVMDispatch_CallObjectAccessorFunction.h
@@ -7,12 +7,12 @@
 #include "RigVMDispatch_CallObjectAccessorFunction.generated.h"
 
 /** Synthetic dispatch injected by the compiler to get a value via an accessor UFunction, not user instantiated */
-USTRUCT(meta = (Hidden, DisplayName = "Call Function", Category="Internal"))
-struct ANIMNEXT_API FRigVMDispatch_CallObjectAccessorFunction : public FRigVMDispatchFactory
+USTRUCT(meta = (Hidden, Category="Internal"))
+struct ANIMNEXT_API FRigVMDispatch_CallObjectAccessorFunctionBase : public FRigVMDispatchFactory
 {
 	GENERATED_BODY()
 
-	FRigVMDispatch_CallObjectAccessorFunction();
+	FRigVMDispatch_CallObjectAccessorFunctionBase();
 
 	static inline const FLazyName ObjectName = FLazyName("Object");
 	static inline const FLazyName FunctionName = FLazyName("Function");
@@ -25,10 +25,36 @@ private:
 	virtual const TArray<FRigVMTemplateArgumentInfo>& GetArgumentInfos() const override;
 	virtual FRigVMTemplateTypeMap OnNewArgumentType(const FName& InArgumentName, TRigVMTypeIndex InTypeIndex) const override;
 	virtual bool IsSingleton() const override { return true; }
+};
+
+/** Script function specialization of FRigVMDispatch_CallObjectAccessorFunctionBase */
+USTRUCT(meta = (Hidden, DisplayName = "Call Native Function"))
+struct ANIMNEXT_API FRigVMDispatch_CallObjectAccessorFunctionNative : public FRigVMDispatch_CallObjectAccessorFunctionBase
+{
+	GENERATED_BODY()
 
+	FRigVMDispatch_CallObjectAccessorFunctionNative();
+
+private:
 	virtual FRigVMFunctionPtr GetDispatchFunctionImpl(const FRigVMTemplateTypeMap& InTypes) const override
 	{
-		return &FRigVMDispatch_CallObjectAccessorFunction::Execute;
+		return &FRigVMDispatch_CallObjectAccessorFunctionNative::Execute;
 	}
 	static void Execute(FRigVMExtendedExecuteContext& InContext, FRigVMMemoryHandleArray Handles, FRigVMPredicateBranchArray RigVMBranches);
 };
+
+/** Script function specialization of FRigVMDispatch_CallObjectAccessorFunctionBase */
+USTRUCT(meta = (Hidden, DisplayName = "Call Script Function"))
+struct ANIMNEXT_API FRigVMDispatch_CallObjectAccessorFunctionScript : public FRigVMDispatch_CallObjectAccessorFunctionBase
+{
+	GENERATED_BODY()
+
+	FRigVMDispatch_CallObjectAccessorFunctionScript();
+
+private:
+	virtual FRigVMFunctionPtr GetDispatchFunctionImpl(const FRigVMTemplateTypeMap& InTypes) const override
+	{
+		return &FRigVMDispatch_CallObjectAccessorFunctionScript::Execute;
+	}
+	static void Execute(FRigVMExtendedExecuteContext& InContext, FRigVMMemoryHandleArray Handles, FRigVMPredicateBranchArray RigVMBranches);
+};
\ No newline at end of file
diff --git a/Engine/Plugins/Experimental/AnimNext/Source/AnimNext/Private/AnimNextDebugDraw.cpp b/Engine/Plugins/Experimental/AnimNext/Source/AnimNext/Private/AnimNextDebugDraw.cpp
--- a/Engine/Plugins/Experimental/AnimNext/Source/AnimNext/Private/AnimNextDebugDraw.cpp
+++ b/Engine/Plugins/Experimental/AnimNext/Source/AnimNext/Private/AnimNextDebugDraw.cpp
@@ -113,6 +113,9 @@ uint32 FAnimNextDebugSceneProxy::GetMemoryFootprint(void) const
 
 FDebugDraw::FCustomSceneProxyDesc::FCustomSceneProxyDesc(UObject* InOwner)
 {
+	check(InOwner);
+	check(InOwner->GetWorld());
+
 	bReceivesDecals = false;
 	bVisibleInReflectionCaptures = false;
 	bVisibleInRealTimeSkyCaptures = false;
diff --git a/Engine/Plugins/Experimental/AnimNext/Source/AnimNext/Private/Module/AnimNextModuleInstance.cpp b/Engine/Plugins/Experimental/AnimNext/Source/AnimNext/Private/Module/AnimNextModuleInstance.cpp
--- a/Engine/Plugins/Experimental/AnimNext/Source/AnimNext/Private/Module/AnimNextModuleInstance.cpp
+++ b/Engine/Plugins/Experimental/AnimNext/Source/AnimNext/Private/Module/AnimNextModuleInstance.cpp
@@ -32,7 +32,10 @@ FAnimNextModuleInstance::FAnimNextModuleInstance(
 	DataInterface = InModule;
 
 #if UE_ENABLE_DEBUG_DRAWING
-	DebugDraw = MakeUnique<UE::AnimNext::Debug::FDebugDraw>(Object);
+	if(Object && Object->GetWorld())
+	{
+		DebugDraw = MakeUnique<UE::AnimNext::Debug::FDebugDraw>(Object);
+	}
 #endif
 }
 
@@ -417,16 +420,17 @@ FModuleInstanceComponent& FAnimNextModuleInstance::AddComponentInternal(int32 Co
 }
 
 #if UE_ENABLE_DEBUG_DRAWING
-FRigVMDrawInterface& FAnimNextModuleInstance::GetDebugDrawInterface()
+FRigVMDrawInterface* FAnimNextModuleInstance::GetDebugDrawInterface()
 {
-	check(DebugDraw.IsValid());
-	return DebugDraw->DrawInterface;
+	return DebugDraw ? &DebugDraw->DrawInterface : nullptr;
 }
 
 void FAnimNextModuleInstance::ShowDebugDrawing(bool bInShowDebugDrawing)
 {
-	check(DebugDraw.IsValid());
-	DebugDraw->SetEnabled(bInShowDebugDrawing);
+	if(DebugDraw)
+	{
+		DebugDraw->SetEnabled(bInShowDebugDrawing);
+	}
 }
 #endif
 
diff --git a/Engine/Plugins/Experimental/AnimNext/Source/AnimNext/Private/Module/ModuleTickFunction.cpp b/Engine/Plugins/Experimental/AnimNext/Source/AnimNext/Private/Module/ModuleTickFunction.cpp
--- a/Engine/Plugins/Experimental/AnimNext/Source/AnimNext/Private/Module/ModuleTickFunction.cpp
+++ b/Engine/Plugins/Experimental/AnimNext/Source/AnimNext/Private/Module/ModuleTickFunction.cpp
@@ -132,7 +132,7 @@ void FModuleEventTickFunction::Run(float InDeltaTime)
 		AnimNextContext.SetOwningObject(ModuleInstance->Object);
 
 #if UE_ENABLE_DEBUG_DRAWING
-		AnimNextContext.SetDrawInterface(&ModuleInstance->GetDebugDrawInterface());
+		AnimNextContext.SetDrawInterface(ModuleInstance->GetDebugDrawInterface());
 #endif
 
 		// Insert our context data for the scope of execution
diff --git a/Engine/Plugins/Experimental/AnimNext/Source/AnimNext/Private/Param/ParamUtils.cpp b/Engine/Plugins/Experimental/AnimNext/Source/AnimNext/Private/Param/ParamUtils.cpp
--- a/Engine/Plugins/Experimental/AnimNext/Source/AnimNext/Private/Param/ParamUtils.cpp
+++ b/Engine/Plugins/Experimental/AnimNext/Source/AnimNext/Private/Param/ParamUtils.cpp
@@ -141,7 +141,7 @@ FParamCompatibility FParamUtils::GetCompatibility(const FAnimNextParamType& InLH
 			{
 				return EParamCompatibility::Compatible_Equal;
 			}
-			if(InLHS.GetValueTypeObject() && InRHS.GetValueTypeObject() && InRHS.GetValueTypeObject()->GetClass()->IsChildOf(InLHS.GetValueTypeObject()->GetClass()))
+			if(InLHS.GetValueTypeObject() && InRHS.GetValueTypeObject() && CastChecked<UClass>(InRHS.GetValueTypeObject())->IsChildOf(CastChecked<UClass>(InLHS.GetValueTypeObject())))
 			{
 				return EParamCompatibility::Compatible_Cast;
 			}
@@ -155,7 +155,7 @@ FParamCompatibility FParamUtils::GetCompatibility(const FAnimNextParamType& InLH
 			{
 				return EParamCompatibility::Compatible_Equal;
 			}
-			if(InLHS.GetValueTypeObject() && InRHS.GetValueTypeObject() && InRHS.GetValueTypeObject()->GetClass()->IsChildOf(InLHS.GetValueTypeObject()->GetClass()))
+			if(InLHS.GetValueTypeObject() && InRHS.GetValueTypeObject() && CastChecked<UClass>(InRHS.GetValueTypeObject())->IsChildOf(CastChecked<UClass>(InLHS.GetValueTypeObject())))
 			{
 				return EParamCompatibility::Compatible_Cast;
 			}
@@ -200,12 +200,14 @@ static bool CanUseFunctionInternal(const UFunction* InFunction, const UClass* In
 	if(FunctionClass->IsChildOf(UBlueprintFunctionLibrary::StaticClass()))
 	{
 		// Check 'hoisted' functions on BPFLs
-		if(!InFunction->HasAllFunctionFlags(FUNC_BlueprintCallable | FUNC_Static | FUNC_Native | FUNC_Public))
+		if(!InFunction->HasAllFunctionFlags(FUNC_BlueprintCallable | FUNC_Static | FUNC_Public))
 		{
 			return false;
 		}
 
-		if(InFunction->NumParms != 2)
+		const bool bValidNative = InFunction->HasAllFunctionFlags(FUNC_Native) && InFunction->NumParms == 2;
+		const bool bValidNonNative = !InFunction->HasAllFunctionFlags(FUNC_Native) && InFunction->NumParms == 3;
+		if(!(bValidNative || bValidNonNative))
 		{
 			return false;
 		}
@@ -250,10 +252,28 @@ static bool CanUseFunctionInternal(const UFunction* InFunction, const UClass* In
 #endif
 			}
 
-			// Check return value
-			if(ParamIndex == 1 && !It->HasAnyPropertyFlags(CPF_ReturnParm))
+			if(bValidNative)
+			{
+				// Check return value
+				if(ParamIndex == 1 && !It->HasAnyPropertyFlags(CPF_ReturnParm))
+				{
+					return false;
+				}
+			}
+			else
 			{
-				return false;
+				// Check world context
+				FObjectProperty* ObjectProperty = CastField<FObjectProperty>(*It);
+				if(ParamIndex == 1 && ObjectProperty == nullptr)
+				{
+					return false;
+				}
+
+				// Check return value
+				if(ParamIndex == 2 && !It->HasAnyPropertyFlags(CPF_ReturnParm))
+				{
+					return false;
+				}
 			}
 
 			OutReturnProperty = *It;
diff --git a/Engine/Plugins/Experimental/AnimNext/Source/AnimNext/Private/Variables/RigUnit_ResolveUniversalObjectLocator.cpp b/Engine/Plugins/Experimental/AnimNext/Source/AnimNext/Private/Variables/RigUnit_ResolveUniversalObjectLocator.cpp
--- a/Engine/Plugins/Experimental/AnimNext/Source/AnimNext/Private/Variables/RigUnit_ResolveUniversalObjectLocator.cpp
+++ b/Engine/Plugins/Experimental/AnimNext/Source/AnimNext/Private/Variables/RigUnit_ResolveUniversalObjectLocator.cpp
@@ -10,7 +10,7 @@ FRigUnit_ResolveUniversalObjectLocator_Execute()
 	const FAnimNextModuleContextData& ContextData = ExecuteContext.GetContextData<FAnimNextModuleContextData>();
 
 	FResolveParams ResolveParams;
-	ResolveParams.Context = ContextData.GetModuleInstance().Object;
+	ResolveParams.Context = ContextData.GetObject();
 	FResolveResult Result = Locator.Resolve(ResolveParams);
 	Object = Result.SyncGet().Object;
 }
\ No newline at end of file
diff --git a/Engine/Plugins/Experimental/AnimNext/Source/AnimNext/Private/Variables/RigVMDispatch_CallHoistedAccessorFunction.cpp b/Engine/Plugins/Experimental/AnimNext/Source/AnimNext/Private/Variables/RigVMDispatch_CallHoistedAccessorFunction.cpp
--- a/Engine/Plugins/Experimental/AnimNext/Source/AnimNext/Private/Variables/RigVMDispatch_CallHoistedAccessorFunction.cpp
+++ b/Engine/Plugins/Experimental/AnimNext/Source/AnimNext/Private/Variables/RigVMDispatch_CallHoistedAccessorFunction.cpp
@@ -4,12 +4,12 @@
 #include "RigVMCore/RigVMStruct.h"
 #include "Variables/AnimNextSoftFunctionPtr.h"
 
-FRigVMDispatch_CallHoistedAccessorFunction::FRigVMDispatch_CallHoistedAccessorFunction()
+FRigVMDispatch_CallHoistedAccessorFunctionBase::FRigVMDispatch_CallHoistedAccessorFunctionBase()
 {
 	FactoryScriptStruct = StaticStruct();
 }
 
-FName FRigVMDispatch_CallHoistedAccessorFunction::GetArgumentNameForOperandIndex(int32 InOperandIndex, int32 InTotalOperands) const
+FName FRigVMDispatch_CallHoistedAccessorFunctionBase::GetArgumentNameForOperandIndex(int32 InOperandIndex, int32 InTotalOperands) const
 {
 	static const FName ArgumentNames[] =
 	{
@@ -21,7 +21,7 @@ FName FRigVMDispatch_CallHoistedAccessorFunction::GetArgumentNameForOperandIndex
 	return ArgumentNames[InOperandIndex];
 }
 
-void FRigVMDispatch_CallHoistedAccessorFunction::RegisterDependencyTypes_NoLock(FRigVMRegistry_NoLock& InRegistry) const
+void FRigVMDispatch_CallHoistedAccessorFunctionBase::RegisterDependencyTypes_NoLock(FRigVMRegistry_NoLock& InRegistry) const
 {
 	static UScriptStruct* const AllowedStructTypes[] =
 	{
@@ -47,7 +47,7 @@ void FRigVMDispatch_CallHoistedAccessorFunction::RegisterDependencyTypes_NoLock(
 	}
 }
 
-const TArray<FRigVMTemplateArgumentInfo>& FRigVMDispatch_CallHoistedAccessorFunction::GetArgumentInfos() const
+const TArray<FRigVMTemplateArgumentInfo>& FRigVMDispatch_CallHoistedAccessorFunctionBase::GetArgumentInfos() const
 {
 	static TArray<FRigVMTemplateArgumentInfo> Infos;
 	if(Infos.IsEmpty())
@@ -67,7 +67,7 @@ const TArray<FRigVMTemplateArgumentInfo>& FRigVMDispatch_CallHoistedAccessorFunc
 	return Infos;
 }
 
-FRigVMTemplateTypeMap FRigVMDispatch_CallHoistedAccessorFunction::OnNewArgumentType(const FName& InArgumentName, TRigVMTypeIndex InTypeIndex) const
+FRigVMTemplateTypeMap FRigVMDispatch_CallHoistedAccessorFunctionBase::OnNewArgumentType(const FName& InArgumentName, TRigVMTypeIndex InTypeIndex) const
 {
 	const FRigVMRegistry_NoLock& Registry = FRigVMRegistry_NoLock::GetForRead();
 	
@@ -78,15 +78,104 @@ FRigVMTemplateTypeMap FRigVMDispatch_CallHoistedAccessorFunction::OnNewArgumentT
 	return Types;
 }
 
-void FRigVMDispatch_CallHoistedAccessorFunction::Execute(FRigVMExtendedExecuteContext& InContext, FRigVMMemoryHandleArray Handles, FRigVMPredicateBranchArray RigVMBranches)
+FRigVMDispatch_CallHoistedAccessorFunctionNative::FRigVMDispatch_CallHoistedAccessorFunctionNative()
 {
-	const UObject* ObjectPtr = *reinterpret_cast<UObject**>(Handles[0].GetData());
+	FactoryScriptStruct = StaticStruct();
+}
+
+void FRigVMDispatch_CallHoistedAccessorFunctionNative::Execute(FRigVMExtendedExecuteContext& InContext, FRigVMMemoryHandleArray Handles, FRigVMPredicateBranchArray RigVMBranches)
+{
+	UObject* ObjectPtr = *reinterpret_cast<UObject**>(Handles[0].GetData());
 	if(ObjectPtr == nullptr)
 	{
 		// Something failed to resolve upstream, OK to just skip this work
 		return;
 	}
 
-	// TODO
+	FAnimNextSoftFunctionPtr* SoftFunctionPtr = reinterpret_cast<FAnimNextSoftFunctionPtr*>(Handles[1].GetData());
+	UFunction* Function = SoftFunctionPtr->SoftObjectPtr.Get();
+	if(Function == nullptr || Function->NumParms != 2)
+	{
+		return;
+	}
+
+	const FProperty* ReturnValueProperty = CastField<FProperty>(Function->GetReturnProperty());
+	if(ReturnValueProperty == nullptr)
+	{
+		return;
+	}
+
+	const FProperty* TargetProperty = Handles[2].GetResolvedProperty();
+	checkSlow(TargetProperty);
+	check(TargetProperty->GetClass() == ReturnValueProperty->GetClass());
+	uint8* TargetAddress = Handles[2].GetData();
+	checkSlow(TargetAddress);
+
+	UObject* CDO = Function->GetOuterUClass()->GetDefaultObject();
+	FFrame Stack(CDO, Function, &ObjectPtr, nullptr, Function->ChildProperties);
+	Function->Invoke(CDO, Stack, TargetAddress);
+}
+
+FRigVMDispatch_CallHoistedAccessorFunctionScript::FRigVMDispatch_CallHoistedAccessorFunctionScript()
+{
+	FactoryScriptStruct = StaticStruct();
+}
+
+void FRigVMDispatch_CallHoistedAccessorFunctionScript::Execute(FRigVMExtendedExecuteContext& InContext, FRigVMMemoryHandleArray Handles, FRigVMPredicateBranchArray RigVMBranches)
+{
+	UObject* ObjectPtr = *reinterpret_cast<UObject**>(Handles[0].GetData());
+	if(ObjectPtr == nullptr)
+	{
+		// Something failed to resolve upstream, OK to just skip this work
+		return;
+	}
+
+	FAnimNextSoftFunctionPtr* SoftFunctionPtr = reinterpret_cast<FAnimNextSoftFunctionPtr*>(Handles[1].GetData());
+	UFunction* Function = SoftFunctionPtr->SoftObjectPtr.Get();
+	if(Function == nullptr || Function->NumParms != 3)
+	{
+		return;
+	}
+
+	const FProperty* ReturnValueProperty = CastField<FProperty>(Function->GetReturnProperty());
+	if(ReturnValueProperty == nullptr)
+	{
+		return;
+	}
+	
+	const FObjectProperty* HoistedProperty = CastField<FObjectProperty>(Function->PropertyLink);
+	if(HoistedProperty == nullptr)
+	{
+		return;
+	}
+
+	check(ObjectPtr->GetClass()->IsChildOf(HoistedProperty->PropertyClass));
+
+	const FObjectProperty* WorldContextProperty = CastField<FObjectProperty>(HoistedProperty->Next);
+	if(WorldContextProperty == nullptr)
+	{
+		return;
+	}
+
+	const FProperty* TargetProperty = Handles[2].GetResolvedProperty();
+	checkSlow(TargetProperty);
+	check(TargetProperty->GetClass() == ReturnValueProperty->GetClass());
+	uint8* TargetAddress = Handles[2].GetData();
+	checkSlow(TargetAddress);
+
+	// Allocate space for function params (script calls require multiple params like the input object and world context to be in one contiguous struct)
+	void* CallParams = FMemory_Alloca_Aligned(Function->GetStructureSize(), Function->GetMinAlignment());
+	Function->InitializeStruct(CallParams);
+	HoistedProperty->SetObjectPropertyValue_InContainer(CallParams, ObjectPtr);
+
+	// Note: currently not setting the world context object here.
+	// If it is needed we can add it, but its generally not thread safe to use
+
+	UObject* CDO = Function->GetOuterUClass()->GetDefaultObject();
+	CDO->ProcessEvent(Function, CallParams);
+
+	// Now copy result out to target address
+	void* SourceAddress = ReturnValueProperty->ContainerPtrToValuePtr<void>(CallParams);
+	ReturnValueProperty->CopyCompleteValue(TargetAddress, SourceAddress);
 }
 
diff --git a/Engine/Plugins/Experimental/AnimNext/Source/AnimNext/Private/Variables/RigVMDispatch_CallObjectAccessorFunction.cpp b/Engine/Plugins/Experimental/AnimNext/Source/AnimNext/Private/Variables/RigVMDispatch_CallObjectAccessorFunction.cpp
--- a/Engine/Plugins/Experimental/AnimNext/Source/AnimNext/Private/Variables/RigVMDispatch_CallObjectAccessorFunction.cpp
+++ b/Engine/Plugins/Experimental/AnimNext/Source/AnimNext/Private/Variables/RigVMDispatch_CallObjectAccessorFunction.cpp
@@ -4,12 +4,12 @@
 #include "RigVMCore/RigVMStruct.h"
 #include "Variables/AnimNextSoftFunctionPtr.h"
 
-FRigVMDispatch_CallObjectAccessorFunction::FRigVMDispatch_CallObjectAccessorFunction()
+FRigVMDispatch_CallObjectAccessorFunctionBase::FRigVMDispatch_CallObjectAccessorFunctionBase()
 {
 	FactoryScriptStruct = StaticStruct();
 }
 
-FName FRigVMDispatch_CallObjectAccessorFunction::GetArgumentNameForOperandIndex(int32 InOperandIndex, int32 InTotalOperands) const
+FName FRigVMDispatch_CallObjectAccessorFunctionBase::GetArgumentNameForOperandIndex(int32 InOperandIndex, int32 InTotalOperands) const
 {
 	static const FName ArgumentNames[] =
 	{
@@ -21,7 +21,7 @@ FName FRigVMDispatch_CallObjectAccessorFunction::GetArgumentNameForOperandIndex(
 	return ArgumentNames[InOperandIndex];
 }
 
-void FRigVMDispatch_CallObjectAccessorFunction::RegisterDependencyTypes_NoLock(FRigVMRegistry_NoLock& InRegistry) const
+void FRigVMDispatch_CallObjectAccessorFunctionBase::RegisterDependencyTypes_NoLock(FRigVMRegistry_NoLock& InRegistry) const
 {
 	static UScriptStruct* const AllowedStructTypes[] =
 	{
@@ -47,7 +47,7 @@ void FRigVMDispatch_CallObjectAccessorFunction::RegisterDependencyTypes_NoLock(F
 	}
 }
 
-const TArray<FRigVMTemplateArgumentInfo>& FRigVMDispatch_CallObjectAccessorFunction::GetArgumentInfos() const
+const TArray<FRigVMTemplateArgumentInfo>& FRigVMDispatch_CallObjectAccessorFunctionBase::GetArgumentInfos() const
 {
 	static TArray<FRigVMTemplateArgumentInfo> Infos;
 	if(Infos.IsEmpty())
@@ -67,7 +67,7 @@ const TArray<FRigVMTemplateArgumentInfo>& FRigVMDispatch_CallObjectAccessorFunct
 	return Infos;
 }
 
-FRigVMTemplateTypeMap FRigVMDispatch_CallObjectAccessorFunction::OnNewArgumentType(const FName& InArgumentName, TRigVMTypeIndex InTypeIndex) const
+FRigVMTemplateTypeMap FRigVMDispatch_CallObjectAccessorFunctionBase::OnNewArgumentType(const FName& InArgumentName, TRigVMTypeIndex InTypeIndex) const
 {
 	const FRigVMRegistry_NoLock& Registry = FRigVMRegistry_NoLock::GetForRead();
 	
@@ -78,15 +78,77 @@ FRigVMTemplateTypeMap FRigVMDispatch_CallObjectAccessorFunction::OnNewArgumentTy
 	return Types;
 }
 
-void FRigVMDispatch_CallObjectAccessorFunction::Execute(FRigVMExtendedExecuteContext& InContext, FRigVMMemoryHandleArray Handles, FRigVMPredicateBranchArray RigVMBranches)
+FRigVMDispatch_CallObjectAccessorFunctionNative::FRigVMDispatch_CallObjectAccessorFunctionNative()
 {
-	const UObject* ObjectPtr = *reinterpret_cast<UObject**>(Handles[0].GetData());
+	FactoryScriptStruct = StaticStruct();
+}
+
+void FRigVMDispatch_CallObjectAccessorFunctionNative::Execute(FRigVMExtendedExecuteContext& InContext, FRigVMMemoryHandleArray Handles, FRigVMPredicateBranchArray RigVMBranches)
+{
+	UObject* ObjectPtr = *reinterpret_cast<UObject**>(Handles[0].GetData());
 	if(ObjectPtr == nullptr)
 	{
 		// Something failed to resolve upstream, OK to just skip this work
 		return;
 	}
 
-	// TODO
+	FAnimNextSoftFunctionPtr* SoftFunctionPtr = reinterpret_cast<FAnimNextSoftFunctionPtr*>(Handles[1].GetData());
+	UFunction* Function = SoftFunctionPtr->SoftObjectPtr.Get();
+	if(Function == nullptr || Function->NumParms != 1)
+	{
+		return;
+	}
+
+	const FProperty* ReturnValueProperty = CastField<FProperty>(Function->GetReturnProperty());
+	if(ReturnValueProperty == nullptr)
+	{
+		return;
+	}
+
+	const FProperty* TargetProperty = Handles[2].GetResolvedProperty();
+	checkSlow(TargetProperty);
+	check(TargetProperty->GetClass() == ReturnValueProperty->GetClass());
+	uint8* TargetAddress = Handles[2].GetData();
+	checkSlow(TargetAddress);
+
+	FFrame Stack(ObjectPtr, Function, nullptr, nullptr, Function->ChildProperties);
+	Function->Invoke(ObjectPtr, Stack, TargetAddress);
+}
+
+FRigVMDispatch_CallObjectAccessorFunctionScript::FRigVMDispatch_CallObjectAccessorFunctionScript()
+{
+	FactoryScriptStruct = StaticStruct();
+}
+
+void FRigVMDispatch_CallObjectAccessorFunctionScript::Execute(FRigVMExtendedExecuteContext& InContext, FRigVMMemoryHandleArray Handles, FRigVMPredicateBranchArray RigVMBranches)
+{
+	UObject* ObjectPtr = *reinterpret_cast<UObject**>(Handles[0].GetData());
+	if(ObjectPtr == nullptr)
+	{
+		// Something failed to resolve upstream, OK to just skip this work
+		return;
+	}
+
+	FAnimNextSoftFunctionPtr* SoftFunctionPtr = reinterpret_cast<FAnimNextSoftFunctionPtr*>(Handles[1].GetData());
+	UFunction* Function = SoftFunctionPtr->SoftObjectPtr.Get();
+	if(Function == nullptr || Function->NumParms != 1)
+	{
+		return;
+	}
+
+	const FProperty* ReturnValueProperty = CastField<FProperty>(Function->GetReturnProperty());
+	if(ReturnValueProperty == nullptr)
+	{
+		return;
+	}
+
+	const FProperty* TargetProperty = Handles[2].GetResolvedProperty();
+	checkSlow(TargetProperty);
+	check(TargetProperty->GetClass() == ReturnValueProperty->GetClass());
+	uint8* TargetAddress = Handles[2].GetData();
+	checkSlow(TargetAddress);
+
+	check(ObjectPtr->GetClass()->IsChildOf(Function->GetOuterUClass()));
+	ObjectPtr->ProcessEvent(Function, TargetAddress);
 }
 
diff --git a/Engine/Plugins/Experimental/AnimNext/Source/AnimNext/Private/Variables/RigVMDispatch_GetObjectProperty.cpp b/Engine/Plugins/Experimental/AnimNext/Source/AnimNext/Private/Variables/RigVMDispatch_GetObjectProperty.cpp
--- a/Engine/Plugins/Experimental/AnimNext/Source/AnimNext/Private/Variables/RigVMDispatch_GetObjectProperty.cpp
+++ b/Engine/Plugins/Experimental/AnimNext/Source/AnimNext/Private/Variables/RigVMDispatch_GetObjectProperty.cpp
@@ -89,7 +89,11 @@ void FRigVMDispatch_GetObjectProperty::Execute(FRigVMExtendedExecuteContext& InC
 
 	FAnimNextFieldPath* FieldPathPtr = reinterpret_cast<FAnimNextFieldPath*>(Handles[1].GetData());
 	const FProperty* SourceProperty = FieldPathPtr->FieldPath.Get();
-	checkSlow(SourceProperty);
+	if(SourceProperty == nullptr)
+	{
+		return;
+	}
+
 	const uint8* SourceAddress = SourceProperty->ContainerPtrToValuePtr<uint8>(ObjectPtr);
 	checkSlow(SourceAddress);
 
diff --git a/Engine/Plugins/Experimental/AnimNext/Source/AnimNextTestSuite/AnimNextTestSuite.Build.cs b/Engine/Plugins/Experimental/AnimNext/Source/AnimNextTestSuite/AnimNextTestSuite.Build.cs
--- a/Engine/Plugins/Experimental/AnimNext/Source/AnimNextTestSuite/AnimNextTestSuite.Build.cs
+++ b/Engine/Plugins/Experimental/AnimNext/Source/AnimNextTestSuite/AnimNextTestSuite.Build.cs
@@ -17,6 +17,7 @@ namespace UnrealBuildTool.Rules
 					"Engine",
 					"AnimNext",
 					"RigVM",
+					"UniversalObjectLocator",
 				}
 			);
 
diff --git a/Engine/Plugins/Experimental/AnimNext/Source/AnimNextTestSuite/Private/AnimNextVariablesTest.cpp b/Engine/Plugins/Experimental/AnimNext/Source/AnimNextTestSuite/Private/AnimNextVariablesTest.cpp
--- a/Engine/Plugins/Experimental/AnimNext/Source/AnimNextTestSuite/Private/AnimNextVariablesTest.cpp
+++ b/Engine/Plugins/Experimental/AnimNext/Source/AnimNextTestSuite/Private/AnimNextVariablesTest.cpp
@@ -10,6 +10,7 @@
 #include "Animation/AnimSequence.h"
 #include "Animation/AnimMontage.h"
 #include "Async/ParallelFor.h"
+#include "Component/AnimNextComponent.h"
 #include "Entries/AnimNextEventGraphEntry.h"
 #include "Module/AnimNextModule.h"
 #include "Module/AnimNextModule_EditorData.h"
@@ -17,6 +18,8 @@
 #include "Factories/Factory.h"
 #include "Module/AnimNextModuleFactory.h"
 #include "Module/RigUnit_AnimNextModuleEvents.h"
+#include "Variables/AnimNextUniversalObjectLocatorBindingData.h"
+#include "Module/AnimNextModuleInstance.h"
 
 // AnimNext Parameters Tests
 
@@ -219,10 +222,10 @@ bool FVariables::RunTest(const FString& InParameters)
 
 	UFactory* Factory = NewObject<UFactory>(GetTransientPackage(), UAnimNextModuleFactory::StaticClass());
 	UAnimNextRigVMAsset* Asset = CastChecked<UAnimNextRigVMAsset>(Factory->FactoryCreateNew(UAnimNextModule::StaticClass(), GetTransientPackage(), TEXT("TestAsset"), RF_Transient, nullptr, nullptr, NAME_None));
-	UE_RETURN_ON_ERROR(Asset != nullptr, "FEditor_Graphs -> Failed to create asset");
+	UE_RETURN_ON_ERROR(Asset != nullptr, "FVariables -> Failed to create asset");
 
 	UAnimNextRigVMAssetEditorData* EditorData = UncookedOnly::FUtils::GetEditorData<UAnimNextRigVMAssetEditorData>(Asset);
-	UE_RETURN_ON_ERROR(EditorData != nullptr, "FEditor_Graphs -> Asset has no editor data.");
+	UE_RETURN_ON_ERROR(EditorData != nullptr, "FVariables -> Asset has no editor data.");
 
 	// Add variables
 	UAnimNextVariableEntry* OperandAEntry = EditorData->AddVariable(TEXT("A"), FAnimNextParamType::GetType<int32>(), TEXT("1"));
@@ -269,10 +272,12 @@ bool FVariables::RunTest(const FString& InParameters)
 	UE_RETURN_ON_ERROR(bLinkExec1Added, TEXT("Could not link variable node exec."));
 
 	bool bLinkExec2Added = Controller->AddLink(TestOpUnitNode->FindPin(FRigVMStruct::ExecuteContextName.ToString()), SetResultNode->FindPin(FRigVMStruct::ExecuteContextName.ToString()));
-	UE_RETURN_ON_ERROR(bLinkExec1Added, TEXT("Could not link variable node exec."));
+	UE_RETURN_ON_ERROR(bLinkExec2Added, TEXT("Could not link variable node exec."));
 
 	URigVMUnitNode* PrintResultUnitNode = Controller->AddUnitNode(FAnimNextTests_PrintResult::StaticStruct());
 	bool bLinkExec3Added = Controller->AddLink(SetResultNode->FindPin(FRigVMStruct::ExecuteContextName.ToString()), PrintResultUnitNode->FindPin(FRigVMStruct::ExecuteContextName.ToString()));
+	UE_RETURN_ON_ERROR(bLinkExec3Added, TEXT("Could not link variable node exec."));
+
 	URigVMVariableNode* GetResultNode = Controller->AddVariableNode("Result", RigVMTypeUtils::Int32Type, nullptr, true, TEXT(""));
 	UE_RETURN_ON_ERROR(GetResultNode != nullptr, TEXT("Could not add get variable node."));
 	bool bLinkResult2Added = Controller->AddLink(GetResultNode->FindPin("Value"), PrintResultUnitNode->FindPin("Result"));
@@ -294,6 +299,130 @@ bool FVariables::RunTest(const FString& InParameters)
 	return true;
 }
 
+
+IMPLEMENT_SIMPLE_AUTOMATION_TEST(FVariables_UOLBindings, "Animation.AnimNext.Variables.UOLBindings", EAutomationTestFlags::EditorContext | EAutomationTestFlags::EngineFilter)
+
+bool FVariables_UOLBindings::RunTest(const FString& InParameters)
+{
+	ON_SCOPE_EXIT{ FUtils::CleanupAfterTests(); };
+
+	UFactory* Factory = NewObject<UFactory>(GetTransientPackage(), UAnimNextModuleFactory::StaticClass());
+	UAnimNextModule* Asset = CastChecked<UAnimNextModule>(Factory->FactoryCreateNew(UAnimNextModule::StaticClass(), GetTransientPackage(), TEXT("TestAsset"), RF_Transient, nullptr, nullptr, NAME_None));
+	UE_RETURN_ON_ERROR(Asset != nullptr, "FVariables_UOLBindings -> Failed to create asset");
+
+	UAnimNextRigVMAssetEditorData* EditorData = UncookedOnly::FUtils::GetEditorData<UAnimNextRigVMAssetEditorData>(Asset);
+	UE_RETURN_ON_ERROR(EditorData != nullptr, "FVariables_UOLBindings -> Asset has no editor data.");
+
+	UniversalObjectLocator::FParseStringParams ParseStringParams;
+	ParseStringParams.Flags |= UniversalObjectLocator::EParseStringFlags::ErrorMessaging;
+
+	// Add variables
+	UAnimNextVariableEntry* PropertyEntry = EditorData->AddVariable(TEXT("Property"), FAnimNextParamType::GetType<int32>(), TEXT("1"));
+	UE_RETURN_ON_ERROR(PropertyEntry != nullptr, TEXT("Could not create new variable in graph."));
+	TInstancedStruct<FAnimNextUniversalObjectLocatorBindingData> BindingA = TInstancedStruct<FAnimNextUniversalObjectLocatorBindingData>::Make();
+	FAnimNextUniversalObjectLocatorBindingData& BindingDataA = BindingA.GetMutable<FAnimNextUniversalObjectLocatorBindingData>();
+	BindingDataA.Type = FAnimNextUniversalObjectLocatorBindingType::Property;
+	BindingDataA.Property = UAnimNextTestFuncLib::StaticClass()->FindPropertyByName(GET_MEMBER_NAME_CHECKED(UAnimNextTestFuncLib, ValueA));
+	UE_RETURN_ON_ERROR(BindingDataA.Property != nullptr, TEXT("Could not find property."));
+	UE::UniversalObjectLocator::FParseStringResult ResultA = BindingDataA.Locator.TryParseString(TEXT("uobj://animobjfunc?&payload0=/Script/AnimNextTestSuite.AnimNextTestFuncLib:GetObj"), ParseStringParams);
+	UE_RETURN_ON_ERROR(ResultA.bSuccess, TEXT("Could not parse UOL."));
+	PropertyEntry->SetBinding(MoveTemp(BindingA));
+
+	UAnimNextVariableEntry* ObjectAccessorEntry = EditorData->AddVariable(TEXT("ObjectAccessor"), FAnimNextParamType::GetType<int32>(), TEXT("2"));
+	UE_RETURN_ON_ERROR(ObjectAccessorEntry != nullptr, TEXT("Could not create new variable in graph."));
+	TInstancedStruct<FAnimNextUniversalObjectLocatorBindingData> BindingB = TInstancedStruct<FAnimNextUniversalObjectLocatorBindingData>::Make();
+	FAnimNextUniversalObjectLocatorBindingData& BindingDataB = BindingB.GetMutable<FAnimNextUniversalObjectLocatorBindingData>();
+	BindingDataB.Type = FAnimNextUniversalObjectLocatorBindingType::Function;
+	BindingDataB.Function = UAnimNextTestFuncLib::StaticClass()->FindFunctionByName(GET_FUNCTION_NAME_CHECKED(UAnimNextTestFuncLib, GetValueB));
+	UE_RETURN_ON_ERROR(BindingDataB.Function != nullptr, TEXT("Could not find function."));
+	UE::UniversalObjectLocator::FParseStringResult ResultB = BindingDataB.Locator.TryParseString(TEXT("uobj://animobjfunc?&payload0=/Script/AnimNextTestSuite.AnimNextTestFuncLib:GetObj"), ParseStringParams);
+	UE_RETURN_ON_ERROR(ResultB.bSuccess, TEXT("Could not parse UOL."));
+	ObjectAccessorEntry->SetBinding(MoveTemp(BindingB));
+
+	UAnimNextVariableEntry* HoistedAccessorEntry = EditorData->AddVariable(TEXT("HoistedAccessor"), FAnimNextParamType::GetType<int32>(), TEXT("12"));
+	UE_RETURN_ON_ERROR(HoistedAccessorEntry != nullptr, TEXT("Could not create new variable in graph."));
+	TInstancedStruct<FAnimNextUniversalObjectLocatorBindingData> BindingC = TInstancedStruct<FAnimNextUniversalObjectLocatorBindingData>::Make();
+	FAnimNextUniversalObjectLocatorBindingData& BindingDataC = BindingC.GetMutable<FAnimNextUniversalObjectLocatorBindingData>();
+	BindingDataC.Type = FAnimNextUniversalObjectLocatorBindingType::HoistedFunction;
+	BindingDataC.Function = UAnimNextTestFuncLib::StaticClass()->FindFunctionByName(GET_FUNCTION_NAME_CHECKED(UAnimNextTestFuncLib, GetValueC));
+	UE_RETURN_ON_ERROR(BindingDataC.Function != nullptr, TEXT("Could not find function."));
+	UE::UniversalObjectLocator::FParseStringResult ResultC = BindingDataC.Locator.TryParseString(TEXT("uobj://animobjfunc?&payload0=/Script/AnimNextTestSuite.AnimNextTestFuncLib:GetObj"), ParseStringParams);
+	UE_RETURN_ON_ERROR(ResultC.bSuccess, TEXT("Could not parse UOL."));
+	HoistedAccessorEntry->SetBinding(MoveTemp(BindingC));
+
+	// Get event graph
+	UAnimNextEventGraphEntry* EventGraph = Cast<UAnimNextEventGraphEntry>(EditorData->FindEntry(TEXT("PrePhysics")));
+	if(EventGraph == nullptr)
+	{
+		EventGraph = EditorData->AddEventGraph(TEXT("PrePhysics"), FRigUnit_AnimNextPrePhysicsEvent::StaticStruct());
+	}
+	UE_RETURN_ON_ERROR(EventGraph != nullptr, TEXT("Could not create new event graph in asset."));
+	
+	URigVMGraph* RigVMGraph = EventGraph->GetRigVMGraph();
+	UE_RETURN_ON_ERROR(RigVMGraph->GetNodes().Num() == 1, TEXT("Unexpected number of nodes in new event graph."));
+
+	URigVMNode* EventNode = RigVMGraph->GetNodes()[0];
+	check(EventNode);
+	URigVMPin* ExecutePin = EventNode->FindPin("ExecuteContext");
+	UE_RETURN_ON_ERROR(ExecutePin != nullptr, TEXT("Could find initial execute pin."));
+
+	URigVMController* Controller = EditorData->GetController(EventGraph->GetRigVMGraph());
+
+	URigVMVariableNode* PropertyEntryNode = Controller->AddVariableNode("Property", RigVMTypeUtils::Int32Type, nullptr, true, TEXT(""));
+	UE_RETURN_ON_ERROR(PropertyEntryNode != nullptr, TEXT("Could not add get variable node."));
+	URigVMVariableNode* ObjectAccessorEntryNode = Controller->AddVariableNode("ObjectAccessor", RigVMTypeUtils::Int32Type, nullptr, true, TEXT(""));
+	UE_RETURN_ON_ERROR(ObjectAccessorEntryNode != nullptr, TEXT("Could not add get variable node."));
+	URigVMVariableNode* HoistedAccessorEntryNode = Controller->AddVariableNode("HoistedAccessor", RigVMTypeUtils::Int32Type, nullptr, true, TEXT(""));
+	UE_RETURN_ON_ERROR(HoistedAccessorEntryNode != nullptr, TEXT("Could not add get variable node."));
+
+	URigVMUnitNode* PrintPropertyUnitNode = Controller->AddUnitNode(FAnimNextTests_PrintResult::StaticStruct());
+	bool bLinkAAdded = Controller->AddLink(PropertyEntryNode->FindPin("Value"), PrintPropertyUnitNode->FindPin("Result"));
+	UE_RETURN_ON_ERROR(bLinkAAdded, TEXT("Could not link variable node."));
+
+	bool bLinkExec1Added = Controller->AddLink(EventNode->FindPin(FRigVMStruct::ExecuteContextName.ToString()), PrintPropertyUnitNode->FindPin(FRigVMStruct::ExecuteContextName.ToString()));
+	UE_RETURN_ON_ERROR(bLinkExec1Added, TEXT("Could not link variable node exec."));
+	
+	URigVMUnitNode* PrintObjectAccessorUnitNode = Controller->AddUnitNode(FAnimNextTests_PrintResult::StaticStruct());
+	bool bLinkBAdded = Controller->AddLink(ObjectAccessorEntryNode->FindPin("Value"), PrintObjectAccessorUnitNode->FindPin("Result"));
+	UE_RETURN_ON_ERROR(bLinkBAdded, TEXT("Could not link variable node."));
+	
+	bool bLinkExec2Added = Controller->AddLink(PrintPropertyUnitNode->FindPin(FRigVMStruct::ExecuteContextName.ToString()), PrintObjectAccessorUnitNode->FindPin(FRigVMStruct::ExecuteContextName.ToString()));
+	UE_RETURN_ON_ERROR(bLinkExec2Added, TEXT("Could not link variable node exec."));
+
+	URigVMUnitNode* PrintHoistedAccessorUnitNode = Controller->AddUnitNode(FAnimNextTests_PrintResult::StaticStruct());
+	bool bLinkCAdded = Controller->AddLink(HoistedAccessorEntryNode->FindPin("Value"), PrintHoistedAccessorUnitNode->FindPin("Result"));
+	UE_RETURN_ON_ERROR(bLinkCAdded, TEXT("Could not link variable node."));
+
+	bool bLinkExec3Added = Controller->AddLink(PrintObjectAccessorUnitNode->FindPin(FRigVMStruct::ExecuteContextName.ToString()), PrintHoistedAccessorUnitNode->FindPin(FRigVMStruct::ExecuteContextName.ToString()));
+	UE_RETURN_ON_ERROR(bLinkExec3Added, TEXT("Could not link variable node exec."));
+
+	TArray<FString> Messages;
+	FRigVMRuntimeSettings RuntimeSettings;
+	RuntimeSettings.SetLogFunction([&Messages](const FRigVMLogSettings& InLogSettings, const FRigVMExecuteContext* InContext, const FString& Message)
+	{
+		Messages.Add(Message);
+	});
+	Asset->GetRigVMExtendedExecuteContext().SetRuntimeSettings(RuntimeSettings);
+	
+	FAnimNextExecuteContext& AnimNextContext = Asset->GetRigVMExtendedExecuteContext().GetPublicDataSafe<FAnimNextExecuteContext>();
+	FAnimNextModuleInstance Instance(Asset, NewObject<UAnimNextComponent>(), EAnimNextModuleInitMethod::None);
+	FAnimNextModuleContextData ContextData(&Instance);
+	FScopedExecuteContextData ContextDataScope(AnimNextContext, ContextData);
+
+	// Execute bindings to fetch data
+	Asset->GetVM()->ExecuteVM(Asset->GetRigVMExtendedExecuteContext(), FRigUnit_AnimNextExecuteBindings::EventName);
+
+	// Run PrePhysics event to print messages
+	Asset->GetVM()->ExecuteVM(Asset->GetRigVMExtendedExecuteContext(), FRigUnit_AnimNextPrePhysicsEvent::EventName);
+
+	UE_RETURN_ON_ERROR(Messages.Num() == 3, TEXT("unexpected number of messages"));
+	UE_RETURN_ON_ERROR(Messages[0] == TEXT("Result = 23"), TEXT("unexpected result message"));
+	UE_RETURN_ON_ERROR(Messages[1] == TEXT("Result = 42"), TEXT("unexpected result message"));
+	UE_RETURN_ON_ERROR(Messages[2] == TEXT("Result = 12345"), TEXT("unexpected result message"));
+
+	return true;
+}
+
 }
 
 #endif	// WITH_DEV_AUTOMATION_TESTS
diff --git a/Engine/Plugins/Experimental/AnimNext/Source/AnimNextTestSuite/Private/AnimNextVariablesTest.h b/Engine/Plugins/Experimental/AnimNext/Source/AnimNextTestSuite/Private/AnimNextVariablesTest.h
--- a/Engine/Plugins/Experimental/AnimNext/Source/AnimNextTestSuite/Private/AnimNextVariablesTest.h
+++ b/Engine/Plugins/Experimental/AnimNext/Source/AnimNextTestSuite/Private/AnimNextVariablesTest.h
@@ -3,11 +3,14 @@
 #pragma once
 
 #include "CoreMinimal.h"
+#include "Kismet/BlueprintFunctionLibrary.h"
 #include "Param/ParamType.h"
 #include "RigVMCore/RigVMStruct.h"
 #include "Templates/SubclassOf.h"
 #include "AnimNextVariablesTest.generated.h"
 
+class UAnimNextComponent;
+
 USTRUCT()
 struct FAnimNextParamTypeTestStruct
 {
@@ -157,3 +160,37 @@ struct FAnimNextTests_PrintResult : public FRigVMStructMutable
 	UPROPERTY(meta=(Input))
 	int32 Result = 0;
 };
+
+UCLASS()
+class UAnimNextTestFuncLib : public UBlueprintFunctionLibrary
+{
+	GENERATED_BODY()
+
+public:
+	UFUNCTION(BlueprintCallable, Category=Test)
+	static UAnimNextTestFuncLib* GetObj(UAnimNextComponent* InObj)
+	{
+		return GetMutableDefault<UAnimNextTestFuncLib>();
+	}
+
+	UFUNCTION(BlueprintCallable, Category=Test)
+	int32 GetValueB()
+	{
+		return ValueB;
+	}
+
+	UFUNCTION(BlueprintCallable, Category=Test)
+	static int32 GetValueC(UAnimNextComponent* InObj)
+	{
+		return GetDefault<UAnimNextTestFuncLib>()->ValueC;
+	}
+
+	UPROPERTY()
+	int32 ValueA = 23;
+
+	UPROPERTY()
+	int32 ValueB = 42;
+
+	UPROPERTY()
+	int32 ValueC = 12345;
+};
\ No newline at end of file
diff --git a/Engine/Plugins/Experimental/AnimNext/Source/AnimNextUncookedOnly/Internal/AnimNextRigVMAssetEditorData.h b/Engine/Plugins/Experimental/AnimNext/Source/AnimNextUncookedOnly/Internal/AnimNextRigVMAssetEditorData.h
--- a/Engine/Plugins/Experimental/AnimNext/Source/AnimNextUncookedOnly/Internal/AnimNextRigVMAssetEditorData.h
+++ b/Engine/Plugins/Experimental/AnimNext/Source/AnimNextUncookedOnly/Internal/AnimNextRigVMAssetEditorData.h
@@ -53,6 +53,7 @@ namespace UE::AnimNext::Tests
 	class FEditor_Variables;
 	class FEditor_AnimGraph_Variables;
 	class FVariables;
+	class FVariables_UOLBindings;
 	class FDataInterfaceCompile;
 }
 
@@ -165,6 +166,7 @@ protected:
 	friend class UE::AnimNext::Tests::FEditor_AnimGraph_Graphs;
 	friend class UE::AnimNext::Tests::FEditor_AnimGraph_Variables;
 	friend class UE::AnimNext::Tests::FVariables;
+	friend class UE::AnimNext::Tests::FVariables_UOLBindings;
 	friend class UE::AnimNext::Tests::FDataInterfaceCompile;
 	friend class UE::AnimNext::Editor::FVariableCustomization;
 	friend class UE::AnimNext::Editor::FAnimNextAssetItemDetails;
diff --git a/Engine/Plugins/Experimental/AnimNext/Source/AnimNextUncookedOnly/Internal/Entries/AnimNextEventGraphEntry.h b/Engine/Plugins/Experimental/AnimNext/Source/AnimNextUncookedOnly/Internal/Entries/AnimNextEventGraphEntry.h
--- a/Engine/Plugins/Experimental/AnimNext/Source/AnimNextUncookedOnly/Internal/Entries/AnimNextEventGraphEntry.h
+++ b/Engine/Plugins/Experimental/AnimNext/Source/AnimNextUncookedOnly/Internal/Entries/AnimNextEventGraphEntry.h
@@ -16,6 +16,7 @@ class UAnimNextEdGraph;
 namespace UE::AnimNext::Tests
 {
 	class FEditor_Variables;
+	class FVariables_UOLBindings;
 	class FEditor_Graphs;
 	class FEditor_AnimGraph_Variables;
 	class FEditor_AnimGraph_Graphs;
@@ -35,6 +36,7 @@ class UAnimNextEventGraphEntry : public UAnimNextRigVMAssetEntry, public IAnimNe
 	friend class UAnimNextRigVMAssetEditorData;
 	friend class UAnimNextModule_EditorData;
 	friend class UE::AnimNext::Tests::FEditor_Variables;
+	friend class UE::AnimNext::Tests::FVariables_UOLBindings;
 	friend class UE::AnimNext::Tests::FEditor_AnimGraph_Variables;
 	friend class UE::AnimNext::Tests::FEditor_Graphs;
 	friend class UE::AnimNext::Tests::FEditor_AnimGraph_Graphs;
diff --git a/Engine/Plugins/Experimental/AnimNext/Source/AnimNextUncookedOnly/Internal/Entries/AnimNextVariableEntry.h b/Engine/Plugins/Experimental/AnimNext/Source/AnimNextUncookedOnly/Internal/Entries/AnimNextVariableEntry.h
--- a/Engine/Plugins/Experimental/AnimNext/Source/AnimNextUncookedOnly/Internal/Entries/AnimNextVariableEntry.h
+++ b/Engine/Plugins/Experimental/AnimNext/Source/AnimNextUncookedOnly/Internal/Entries/AnimNextVariableEntry.h
@@ -36,6 +36,7 @@ namespace UE::AnimNext::Editor
 namespace UE::AnimNext::Tests
 {
 	class FEditor_Variables;
+	class FVariables_UOLBindings;
 	class FEditor_AnimGraph_Variables;
 	class FDataInterfaceCompile;
 }
@@ -57,6 +58,7 @@ class UAnimNextVariableEntry : public UAnimNextRigVMAssetEntry, public IAnimNext
 	friend class UAnimNextStateTreeTreeEditorData;
 	friend class UE::AnimNext::Tests::FEditor_Variables;
 	friend class UE::AnimNext::Tests::FEditor_AnimGraph_Variables;
+	friend class UE::AnimNext::Tests::FVariables_UOLBindings;
 	friend class UE::AnimNext::Tests::FDataInterfaceCompile;
 	friend class FAnimationAnimNextRuntimeTest_Variables;
 	friend class UE::AnimNext::Editor::FVariableCustomization;
@@ -99,6 +101,7 @@ class UAnimNextVariableEntry : public UAnimNextRigVMAssetEntry, public IAnimNext
 	virtual const FInstancedPropertyBag& GetPropertyBag() const override;
 	virtual bool GetDefaultValue(const FProperty*& OutProperty, TConstArrayView<uint8>& OutValue) const override;
 	virtual void SetBindingType(UScriptStruct* InBindingTypeStruct, bool bSetupUndoRedo = true) override;
+	virtual void SetBinding(TInstancedStruct<FAnimNextVariableBindingData>&& InBinding, bool bSetupUndoRedo = true) override;
 	virtual TConstStructView<FAnimNextVariableBindingData> GetBinding() const override;
 
 	// Set the default value
diff --git a/Engine/Plugins/Experimental/AnimNext/Source/AnimNextUncookedOnly/Private/Entries/AnimNextVariableEntry.cpp b/Engine/Plugins/Experimental/AnimNext/Source/AnimNextUncookedOnly/Private/Entries/AnimNextVariableEntry.cpp
--- a/Engine/Plugins/Experimental/AnimNext/Source/AnimNextUncookedOnly/Private/Entries/AnimNextVariableEntry.cpp
+++ b/Engine/Plugins/Experimental/AnimNext/Source/AnimNextUncookedOnly/Private/Entries/AnimNextVariableEntry.cpp
@@ -172,6 +172,18 @@ void UAnimNextVariableEntry::SetBindingType(UScriptStruct* InBindingTypeStruct,
 	BroadcastModified(EAnimNextEditorDataNotifType::VariableBindingChanged);
 }
 
+void UAnimNextVariableEntry::SetBinding(TInstancedStruct<FAnimNextVariableBindingData>&& InBinding, bool bSetupUndoRedo)
+{
+	if(bSetupUndoRedo)
+	{
+		Modify();
+	}
+
+	Binding.BindingData = MoveTemp(InBinding);
+
+	BroadcastModified(EAnimNextEditorDataNotifType::VariableBindingChanged);
+}
+
 TConstStructView<FAnimNextVariableBindingData> UAnimNextVariableEntry::GetBinding() const
 {
 	return Binding.BindingData;
diff --git a/Engine/Plugins/Experimental/AnimNext/Source/AnimNextUncookedOnly/Private/Variables/UniversalObjectLocatorBindingType.cpp b/Engine/Plugins/Experimental/AnimNext/Source/AnimNextUncookedOnly/Private/Variables/UniversalObjectLocatorBindingType.cpp
--- a/Engine/Plugins/Experimental/AnimNext/Source/AnimNextUncookedOnly/Private/Variables/UniversalObjectLocatorBindingType.cpp
+++ b/Engine/Plugins/Experimental/AnimNext/Source/AnimNextUncookedOnly/Private/Variables/UniversalObjectLocatorBindingType.cpp
@@ -201,6 +201,7 @@ TSharedRef<SWidget> FUniversalObjectLocatorBindingType::CreateEditWidget(const T
 			InPropertyHandle->AddChildStructure(StructProvider.ToSharedRef());
 			FieldIterator = MakeUnique<FFieldIterator>(InType);
 			PropertyHandle = InPropertyHandle;
+			FilterType = InType;
 
 			LocatorHandle = InPropertyHandle->GetChildHandle(GET_MEMBER_NAME_CHECKED(FAnimNextUniversalObjectLocatorBindingData, Locator));
 			check(LocatorHandle.IsValid());
@@ -329,17 +330,28 @@ TSharedRef<SWidget> FUniversalObjectLocatorBindingType::CreateEditWidget(const T
 					ContainerMap.Add(Handle, CachedContainers.Num() - 1);
 				}
 
-				// Find any native UBlueprintFunctionLibrary classes to extend this class
+				// Find any UBlueprintFunctionLibrary classes to extend this class
 				TArray<UClass*> Classes;
 				GetDerivedClasses(UBlueprintFunctionLibrary::StaticClass(), Classes, true);
-				Classes.Add(const_cast<UClass*>(Class));
 				for (const UClass* LibraryClass : Classes)
 				{
-					if(LibraryClass->HasAnyClassFlags(CLASS_Abstract) || !LibraryClass->HasAnyClassFlags(CLASS_Native))
+					if(LibraryClass->HasAnyClassFlags(CLASS_Abstract))
 					{
 						continue;
 					}
 
+					if(LibraryClass->HasAnyClassFlags(CLASS_CompiledFromBlueprint))
+					{
+						// Make sure we skip any out-of-date or skeleton classes
+						// It appears that the only way to detect skeleton classes is via the SKEL_ prefix
+						if(LibraryClass->HasAnyClassFlags(CLASS_NewerVersionExists) ||
+							LibraryClass->GetName().Contains(TEXT("SKEL_")) ||
+							LibraryClass->GetName().Contains(TEXT("REINST_")))
+						{
+							continue;
+						}
+					}
+
 					auto PassesFilterChecks = [this](const FProperty* InProperty)
 					{
 						if(InProperty && FilterType.IsValid())
@@ -470,7 +482,7 @@ TSharedRef<SWidget> FUniversalObjectLocatorBindingType::CreateEditWidget(const T
 				};
 
 				TArray<FFieldVariant> Result;
-				for (TFieldIterator<FProperty> PropertyIt(InStruct, EFieldIteratorFlags::IncludeSuper); PropertyIt; ++PropertyIt)
+				for (TFieldIterator<FProperty> PropertyIt(InStruct, EFieldIteratorFlags::IncludeSuper, EFieldIteratorFlags::ExcludeDeprecated, EFieldIteratorFlags::IncludeInterfaces); PropertyIt; ++PropertyIt)
 				{
 					FProperty* Property = *PropertyIt;
 					if (FParamUtils::CanUseProperty(Property))
@@ -481,7 +493,7 @@ TSharedRef<SWidget> FUniversalObjectLocatorBindingType::CreateEditWidget(const T
 						}
 					}
 				}
-				for (TFieldIterator<UFunction> FunctionIt(InStruct, EFieldIteratorFlags::IncludeSuper); FunctionIt; ++FunctionIt)
+				for (TFieldIterator<UFunction> FunctionIt(InStruct, EFieldIteratorFlags::IncludeSuper, EFieldIteratorFlags::ExcludeDeprecated, EFieldIteratorFlags::IncludeInterfaces); FunctionIt; ++FunctionIt)
 				{
 					UFunction* Function = *FunctionIt;
 					if (FParamUtils::CanUseFunction(Function, CastChecked<UClass>(CurrentStruct)))
@@ -771,13 +783,31 @@ void FUniversalObjectLocatorBindingType::BuildBindingGraphFragment(const FRigVMC
 			}
 			case FAnimNextUniversalObjectLocatorBindingType::Function:
 			{
-				GetterNode = InArgs.Controller->AddTemplateNode(FRigVMDispatch_CallObjectAccessorFunction().GetTemplateNotation(), FVector2D(200.0f, YOffset + AccumulatedYOffset), TEXT(""), false);
+				UFunction* Function = LocatorBinding.Function.Get();
+				if(Function == nullptr)
+				{
+					InSettings.ReportErrorf(TEXT("Could not resolve function call %s"), *LocatorBinding.Function.ToString());
+					return;
+				}
+
+				// Determine the call specialization we need (native vs script)
+				FName Notation;
+				if(Function->HasAllFunctionFlags(FUNC_Native))
+				{
+					Notation = FRigVMDispatch_CallObjectAccessorFunctionNative().GetTemplateNotation();
+				}
+				else
+				{
+					Notation = FRigVMDispatch_CallObjectAccessorFunctionScript().GetTemplateNotation();
+				}
+
+				GetterNode = InArgs.Controller->AddTemplateNode(Notation, FVector2D(200.0f, YOffset + AccumulatedYOffset), TEXT(""), false);
 				if(GetterNode == nullptr)
 				{
 					InSettings.ReportError(TEXT("Could not spawn Call Object Accessor Function node"));
 					return;
 				}
-				URigVMPin* FunctionPin = GetterNode->FindPin(FRigVMDispatch_CallObjectAccessorFunction::FunctionName.ToString());
+				URigVMPin* FunctionPin = GetterNode->FindPin(FRigVMDispatch_CallObjectAccessorFunctionBase::FunctionName.ToString());
 				if(FunctionPin == nullptr)
 				{
 					InSettings.ReportError(TEXT("Could not find Function pin"));
@@ -797,13 +827,31 @@ void FUniversalObjectLocatorBindingType::BuildBindingGraphFragment(const FRigVMC
 			}
 			case FAnimNextUniversalObjectLocatorBindingType::HoistedFunction:
 			{
-				GetterNode = InArgs.Controller->AddTemplateNode(FRigVMDispatch_CallHoistedAccessorFunction().GetTemplateNotation(), FVector2D(200.0f, YOffset + AccumulatedYOffset), TEXT(""), false);
+				UFunction* Function = LocatorBinding.Function.Get();
+				if(Function == nullptr)
+				{
+					InSettings.ReportErrorf(TEXT("Could not resolve hoisted function call %s"), *LocatorBinding.Function.ToString());
+					return;
+				}
+
+				// Determine the call specialization we need (native vs script)
+				FName Notation;
+				if(Function->HasAllFunctionFlags(FUNC_Native))
+				{
+					Notation = FRigVMDispatch_CallHoistedAccessorFunctionNative().GetTemplateNotation();
+				}
+				else
+				{
+					Notation = FRigVMDispatch_CallHoistedAccessorFunctionScript().GetTemplateNotation();
+				}
+
+				GetterNode = InArgs.Controller->AddTemplateNode(Notation, FVector2D(200.0f, YOffset + AccumulatedYOffset), TEXT(""), false);
 				if(GetterNode == nullptr)
 				{
 					InSettings.ReportError(TEXT("Could not spawn Call Hoisted Accessor Function node"));
 					return;
 				}
-				URigVMPin* FunctionPin = GetterNode->FindPin(FRigVMDispatch_CallObjectAccessorFunction::FunctionName.ToString());
+				URigVMPin* FunctionPin = GetterNode->FindPin(FRigVMDispatch_CallObjectAccessorFunctionBase::FunctionName.ToString());
 				if(FunctionPin == nullptr)
 				{
 					InSettings.ReportError(TEXT("Could not find Function pin"));
diff --git a/Engine/Plugins/Experimental/AnimNext/Source/AnimNextUncookedOnly/Public/Variables/IAnimNextRigVMVariableInterface.h b/Engine/Plugins/Experimental/AnimNext/Source/AnimNextUncookedOnly/Public/Variables/IAnimNextRigVMVariableInterface.h
--- a/Engine/Plugins/Experimental/AnimNext/Source/AnimNextUncookedOnly/Public/Variables/IAnimNextRigVMVariableInterface.h
+++ b/Engine/Plugins/Experimental/AnimNext/Source/AnimNextUncookedOnly/Public/Variables/IAnimNextRigVMVariableInterface.h
@@ -49,10 +49,13 @@ public:
 	}
 
 	virtual bool GetDefaultValue(const FProperty*& OutProperty, TConstArrayView<uint8>& OutValue) const = 0;
-	
-	// Set the binding type for this variable
-	virtual void SetBindingType(UScriptStruct* InBindingTypeStruct, bool bSetupUndoRedo) = 0;
-	
+
+	// Set the binding type for this variable (initializes it to default if the struct type is valid)
+	virtual void SetBindingType(UScriptStruct* InBindingTypeStruct, bool bSetupUndoRedo = true) = 0;
+
+	// Set the binding for this variable
+	virtual void SetBinding(TInstancedStruct<FAnimNextVariableBindingData>&& InBinding, bool bSetupUndoRedo = true) = 0;
+
 	// Get the binding for this variable, if any
 	virtual TConstStructView<FAnimNextVariableBindingData> GetBinding() const = 0;
 
diff --git a/Engine/Plugins/Experimental/AnimNextAnimGraph/Source/AnimNextAnimGraph/Private/TraitCore/ExecutionContext.cpp b/Engine/Plugins/Experimental/AnimNextAnimGraph/Source/AnimNextAnimGraph/Private/TraitCore/ExecutionContext.cpp
--- a/Engine/Plugins/Experimental/AnimNextAnimGraph/Source/AnimNextAnimGraph/Private/TraitCore/ExecutionContext.cpp
+++ b/Engine/Plugins/Experimental/AnimNextAnimGraph/Source/AnimNextAnimGraph/Private/TraitCore/ExecutionContext.cpp
@@ -589,7 +589,7 @@ namespace UE::AnimNext
 	}
 
 #if UE_ENABLE_DEBUG_DRAWING
-	FRigVMDrawInterface& FExecutionContext::GetDebugDrawInterface()
+	FRigVMDrawInterface* FExecutionContext::GetDebugDrawInterface()
 	{
 		check(IsBound());
 		return RootGraphInstance->GetModuleInstance()->GetDebugDrawInterface();
diff --git a/Engine/Plugins/Experimental/AnimNextAnimGraph/Source/AnimNextAnimGraph/Public/TraitCore/ExecutionContext.h b/Engine/Plugins/Experimental/AnimNextAnimGraph/Source/AnimNextAnimGraph/Public/TraitCore/ExecutionContext.h
--- a/Engine/Plugins/Experimental/AnimNextAnimGraph/Source/AnimNextAnimGraph/Public/TraitCore/ExecutionContext.h
+++ b/Engine/Plugins/Experimental/AnimNextAnimGraph/Source/AnimNextAnimGraph/Public/TraitCore/ExecutionContext.h
@@ -215,7 +215,7 @@ namespace UE::AnimNext
 
 #if UE_ENABLE_DEBUG_DRAWING
 		// Get the debug draw interface
-		FRigVMDrawInterface& GetDebugDrawInterface();
+		FRigVMDrawInterface* GetDebugDrawInterface();
 #endif
 	private:
 		// No copy or move
