AnimNext: Fix crash/assert when sync loading AnimNext assets in PIE
Dont broadcast module compiled delegate for newly loaded assets

diff --git a/Engine/Plugins/Experimental/AnimNext/Source/AnimNextUncookedOnly/Internal/AnimNextRigVMAssetEditorData.h b/Engine/Plugins/Experimental/AnimNext/Source/AnimNextUncookedOnly/Internal/AnimNextRigVMAssetEditorData.h
--- a/Engine/Plugins/Experimental/AnimNext/Source/AnimNextUncookedOnly/Internal/AnimNextRigVMAssetEditorData.h
+++ b/Engine/Plugins/Experimental/AnimNext/Source/AnimNextUncookedOnly/Internal/AnimNextRigVMAssetEditorData.h
@@ -522,5 +522,5 @@ protected:
 	bool bCompileInDebugMode = false;
 	bool bSuspendPythonMessagesForRigVMClient = true;
 	bool bSuspendEditorDataNotifications = false;
-	
+	bool bSuspendCompilationNotifications = false;
 };
diff --git a/Engine/Plugins/Experimental/AnimNext/Source/AnimNextUncookedOnly/Private/AnimNextRigVMAssetEditorData.cpp b/Engine/Plugins/Experimental/AnimNext/Source/AnimNextUncookedOnly/Private/AnimNextRigVMAssetEditorData.cpp
--- a/Engine/Plugins/Experimental/AnimNext/Source/AnimNextUncookedOnly/Private/AnimNextRigVMAssetEditorData.cpp
+++ b/Engine/Plugins/Experimental/AnimNext/Source/AnimNextUncookedOnly/Private/AnimNextRigVMAssetEditorData.cpp
@@ -307,6 +307,8 @@ void UAnimNextRigVMAssetEditorData::HandlePackageDone()
 
 	ReconstructAllNodes(); // If this is not executed on a node for whatever reason, it will appear transparent in the editor
 
+	TGuardValue<bool> DisableCompilationNotifications(bSuspendCompilationNotifications, true);
+
 	RecompileVM();
 }
 
diff --git a/Engine/Plugins/Experimental/AnimNext/Source/AnimNextUncookedOnly/Private/Module/AnimNextModule_EditorData.cpp b/Engine/Plugins/Experimental/AnimNext/Source/AnimNextUncookedOnly/Private/Module/AnimNextModule_EditorData.cpp
--- a/Engine/Plugins/Experimental/AnimNext/Source/AnimNextUncookedOnly/Private/Module/AnimNextModule_EditorData.cpp
+++ b/Engine/Plugins/Experimental/AnimNext/Source/AnimNextUncookedOnly/Private/Module/AnimNextModule_EditorData.cpp
@@ -22,7 +22,10 @@ void UAnimNextModule_EditorData::RecompileVM()
 	}
 
 #if WITH_EDITOR
-	UAnimNextModule::OnModuleCompiled().Broadcast(UE::AnimNext::UncookedOnly::FUtils::GetAsset<UAnimNextModule>(this));
+	if(!bSuspendCompilationNotifications)
+	{
+		UAnimNextModule::OnModuleCompiled().Broadcast(UE::AnimNext::UncookedOnly::FUtils::GetAsset<UAnimNextModule>(this));
+	}
 #endif
 }
 
