AnimNext: Fix memory stomp when changing a variable's type from the details panel

diff --git a/Engine/Plugins/Experimental/AnimNext/Source/AnimNextUncookedOnly/Internal/Entries/AnimNextVariableEntry.h b/Engine/Plugins/Experimental/AnimNext/Source/AnimNextUncookedOnly/Internal/Entries/AnimNextVariableEntry.h
--- a/Engine/Plugins/Experimental/AnimNext/Source/AnimNextUncookedOnly/Internal/Entries/AnimNextVariableEntry.h
+++ b/Engine/Plugins/Experimental/AnimNext/Source/AnimNextUncookedOnly/Internal/Entries/AnimNextVariableEntry.h
@@ -78,6 +78,9 @@ class UAnimNextVariableEntry : public UAnimNextRigVMAssetEntry, public IAnimNext
 	// UObject interface
 	virtual void Serialize(FArchive& Ar) override;
 	virtual void PostLoad() override;
+#if WITH_EDITOR
+	virtual void PostEditChangeProperty(FPropertyChangedEvent& PropertyChangedEvent) override;
+#endif
 
 	// IAnimNextRigVMExportInterface interface
 	virtual FAnimNextParamType GetExportType() const override;
diff --git a/Engine/Plugins/Experimental/AnimNext/Source/AnimNextUncookedOnly/Private/Entries/AnimNextVariableEntry.cpp b/Engine/Plugins/Experimental/AnimNext/Source/AnimNextUncookedOnly/Private/Entries/AnimNextVariableEntry.cpp
--- a/Engine/Plugins/Experimental/AnimNext/Source/AnimNextUncookedOnly/Private/Entries/AnimNextVariableEntry.cpp
+++ b/Engine/Plugins/Experimental/AnimNext/Source/AnimNextUncookedOnly/Private/Entries/AnimNextVariableEntry.cpp
@@ -237,4 +237,21 @@ void UAnimNextVariableEntry::PostLoad()
 			NewDesc->CachedProperty->CopyCompleteValue(NewValue, OldValue);
 		}
 	}
-}
\ No newline at end of file
+}
+
+#if WITH_EDITOR
+
+void UAnimNextVariableEntry::PostEditChangeProperty(FPropertyChangedEvent& PropertyChangedEvent)
+{
+	if(PropertyChangedEvent.Property && PropertyChangedEvent.Property->GetName() == GET_MEMBER_NAME_CHECKED(UAnimNextVariableEntry, Type))
+	{
+		SetType(Type);
+	}
+	else
+	{
+		// Call super to broadcast general changed event
+		Super::PostEditChangeProperty(PropertyChangedEvent);
+	}
+}
+
+#endif
\ No newline at end of file
