Add C++ method to set a variable on an AnimNextComponent

diff --git a/Engine/Plugins/Experimental/AnimNext/Source/AnimNext/Internal/Component/AnimNextComponent.h b/Engine/Plugins/Experimental/AnimNext/Source/AnimNext/Internal/Component/AnimNextComponent.h
--- a/Engine/Plugins/Experimental/AnimNext/Source/AnimNext/Internal/Component/AnimNextComponent.h
+++ b/Engine/Plugins/Experimental/AnimNext/Source/AnimNext/Internal/Component/AnimNextComponent.h
@@ -18,7 +18,7 @@ class UAnimNextComponentWorldSubsystem;
 namespace UE::AnimNext
 {
 	struct FProxyVariablesContext;
-};
+}
 
 namespace UE::AnimNext::UncookedOnly
 {
@@ -57,6 +57,125 @@ public:
 	UFUNCTION(BlueprintCallable, Category = "AnimNext", CustomThunk, meta = (CustomStructureParam = Value, UnsafeDuringActorConstruction))
 	ANIMNEXT_API void SetVariable(UPARAM(meta = (CustomWidget = "VariableName")) FName Name, int32 Value);
 
+	/**
+	 * Module variable value getters.
+	 * Numeric types (bool, (u)int32, (u)int64, float, double) support type conversion.
+	 * @param    Name     The name of the variable to get
+	 * @return            The variable value on success, None otherwise
+	 */
+	TOptional<bool> GetVariableBool(const FName Name) const;
+	TOptional<uint8> GetVariableByte(const FName Name) const;
+	TOptional<int32> GetVariableInt32(const FName Name) const;
+	TOptional<uint32> GetVariableUInt32(const FName Name) const;
+	TOptional<int64> GetVariableInt64(const FName Name) const;
+	TOptional<uint64> GetVariableUInt64(const FName Name) const;
+	TOptional<float> GetVariableFloat(const FName Name) const;
+	TOptional<double> GetVariableDouble(const FName Name) const;
+	TOptional<FName> GetVariableName(const FName Name) const;
+	TOptional<FString> GetVariableString(const FName Name) const;
+	TOptional<uint8> GetVariableEnum(const FName Name, const UEnum* RequestedEnum) const;
+	TOptional<FStructView> GetVariableStruct(const FName Name, const UScriptStruct* RequestedStruct = nullptr) const;
+	TOptional<UObject*> GetVariableObject(const FName Name, const UClass* RequestedClass = nullptr) const;
+	TOptional<UClass*> GetVariableClass(const FName Name) const;
+	TOptional<FSoftObjectPath> GetVariableSoftPath(const FName Name) const;
+
+	/** @return enum value of specified type. */
+	template <typename T>
+	TOptional<T> GetVariableEnum(const FName Name) const
+	{
+		static_assert(TIsEnum<T>::Value, "Should only call this with enum types");
+
+		TOptional<uint8> Result = GetVariableEnum(Name, StaticEnum<T>());
+		if (!Result.IsSet())
+		{
+			return TOptional<T>();
+		}
+		return TOptional<T>(static_cast<T>(Result.GetValue()));
+	}
+
+	/** @return struct reference of specified type. */
+	template <typename T>
+	TOptional<T*> GetVariableStruct(const FName Name) const
+	{
+		TOptional<FStructView> Result = GetVariableStruct(Name, TBaseStructure<T>::Get());
+		if (!Result.IsSet())
+		{
+			return TOptional<T*>();
+		}
+		if (T* ValuePtr = Result.GetValue().GetPtr<T>())
+		{
+			return TOptional<T*>(ValuePtr);
+		}
+		return TOptional<T*>();
+	}
+
+	/** @return object pointer value of specified type. */
+	template <typename T>
+	TOptional<T*> GetVariableObject(const FName Name) const
+	{
+		static_assert(TIsDerivedFrom<T, UObject>::Value, "Should only call this with object types");
+
+		TOptional<UObject*> Result = GetVariableObject(Name, T::StaticClass());
+		if (!Result.IsSet())
+		{
+			return TOptional<T*>();
+		}
+		if (Result.GetValue() == nullptr)
+		{
+			return TOptional<T*>(nullptr);
+		}
+		if (T* Object = Cast<T>(Result.GetValue()))
+		{
+			return TOptional<T*>(Object);
+		}
+		return TOptional<T*>();
+	}
+
+	/**
+	 * Module variable value setters.
+	 * @param    Name     The name of the variable to set
+	 * @param    Value    The value to set the variable to
+	 * @return            True on success, false otherwise
+	 */
+	bool SetVariableBool(const FName Name, const bool bInValue);
+	bool SetVariableByte(const FName Name, const uint8 InValue);
+	bool SetVariableInt32(const FName Name, const int32 InValue);
+	bool SetVariableUInt32(const FName Name, const uint32 InValue);
+	bool SetVariableInt64(const FName Name, const int64 InValue);
+	bool SetVariableUInt64(const FName Name, const uint64 InValue);
+	bool SetVariableFloat(const FName Name, const float InValue);
+	bool SetVariableDouble(const FName Name, const double InValue);
+	bool SetVariableName(const FName Name, const FName InValue);
+	bool SetVariableString(const FName Name, const FString& InValue);
+	bool SetVariableEnum(const FName Name, const uint8 InValue, const UEnum* Enum);
+	bool SetVariableStruct(const FName Name, FConstStructView InValue);
+	bool SetVariableObject(const FName Name, UObject* InValue);
+	bool SetVariableClass(const FName Name, UClass* InValue);
+	bool SetVariableSoftPath(const FName Name, const FSoftObjectPath& InValue);
+
+	/** Sets enum variable with specified type. */
+	template <typename T>
+	bool SetVariableEnum(const FName Name, const T InValue)
+	{
+		static_assert(TIsEnum<T>::Value, "Should only call this with enum types");
+		return SetVariableEnum(Name, static_cast<uint8>(InValue), StaticEnum<T>());
+	}
+
+	/** Sets struct variable with specified type. */
+	template <typename T>
+	bool SetVariableStruct(const FName Name, const T& InValue)
+	{
+		return SetVariableStruct(Name, FConstStructView::Make(InValue));
+	}
+
+	/** Sets object pointer variable with specified type. */
+	template <typename T>
+	bool SetVariableObject(const FName Name, T* InValue)
+	{
+		static_assert(TIsDerivedFrom<T, UObject>::Value, "Should only call this with object types");
+		return SetVariableObject(Name, (UObject*)InValue);
+	}
+
 	// Enable or disable this component's update
 	UFUNCTION(BlueprintCallable, Category = "AnimNext")
 	ANIMNEXT_API void SetEnabled(bool bEnabled);
@@ -87,7 +206,7 @@ private:
 	UE::AnimNext::FModuleHandle ModuleHandle;
 
 	// Lock for public variables proxy
-	FRWLock PublicVariablesLock;
+	mutable FRWLock PublicVariablesLock;
 
 	// Proxy public variables
 	UPROPERTY()
diff --git a/Engine/Plugins/Experimental/AnimNext/Source/AnimNext/Private/Component/AnimNextComponent.cpp b/Engine/Plugins/Experimental/AnimNext/Source/AnimNext/Private/Component/AnimNextComponent.cpp
--- a/Engine/Plugins/Experimental/AnimNext/Source/AnimNext/Private/Component/AnimNextComponent.cpp
+++ b/Engine/Plugins/Experimental/AnimNext/Source/AnimNext/Private/Component/AnimNextComponent.cpp
@@ -148,18 +148,638 @@ DEFINE_FUNCTION(UAnimNextComponent::execSetVariable)
 		return;
 	}
 
+	bool bIsTypeMatching = false;
+
 	P_NATIVE_BEGIN;
 
 	{
 		FRWScopeLock Lock(P_THIS->PublicVariablesLock, SLT_Write);
+
 		TConstArrayView<FPropertyBagPropertyDesc> ProxyDescs = P_THIS->PublicVariablesProxy.Data.GetPropertyBagStruct()->GetPropertyDescs();
-		const void* ValuePtr = ValueProp->ContainerPtrToValuePtr<void>(ContainerPtr);
-		ProxyDescs[*IndexPtr].CachedProperty->SetValue_InContainer(P_THIS->PublicVariablesProxy.Data.GetMutableValue().GetMemory(), ValuePtr);
-		P_THIS->PublicVariablesProxy.DirtyFlags[*IndexPtr] = true;
-		P_THIS->PublicVariablesProxy.bIsDirty = true;
+		const FProperty* CachedProperty = ProxyDescs[*IndexPtr].CachedProperty;
+
+		bIsTypeMatching = CachedProperty->GetClass() == ValueProp->GetClass();
+		if (bIsTypeMatching)
+		{
+			const void* ValuePtr = ValueProp->ContainerPtrToValuePtr<void>(ContainerPtr);
+			CachedProperty->SetValue_InContainer(P_THIS->PublicVariablesProxy.Data.GetMutableValue().GetMemory(), ValuePtr);
+			P_THIS->PublicVariablesProxy.DirtyFlags[*IndexPtr] = true;
+			P_THIS->PublicVariablesProxy.bIsDirty = true;
+		}
 	}
 
 	P_NATIVE_END;
+
+	if (!bIsTypeMatching)
+	{
+		FBlueprintExceptionInfo ExceptionInfo(
+			EBlueprintExceptionType::NonFatalError,
+			FText::Format(NSLOCTEXT("AnimNextComponent", "AnimNextComponent_SetVariableTypeMismatch", "Incompatible type supplied for variable '{0}'"), FText::FromName(Name))
+		);
+
+		FBlueprintCoreDelegates::ThrowScriptException(P_THIS, Stack, ExceptionInfo);
+		return;
+	}
+}
+
+TOptional<bool> UAnimNextComponent::GetVariableBool(const FName Name) const
+{
+	const int32* IndexPtr = PublicVariablesProxyMap.Find(Name);
+	if (IndexPtr == nullptr)
+	{
+		return TOptional<bool>();
+	}
+
+	FRWScopeLock Lock(PublicVariablesLock, SLT_ReadOnly);
+
+	TConstArrayView<FPropertyBagPropertyDesc> ProxyDescs = PublicVariablesProxy.Data.GetPropertyBagStruct()->GetPropertyDescs();
+	const FPropertyBagPropertyDesc& PropertyDesc = ProxyDescs[*IndexPtr];
+
+	TValueOrError<bool, EPropertyBagResult> Result = PublicVariablesProxy.Data.GetValueBool(PropertyDesc);
+	return Result.HasValue() ? TOptional<bool>(Result.GetValue()) : TOptional<bool>();
+}
+
+TOptional<uint8> UAnimNextComponent::GetVariableByte(const FName Name) const
+{
+	const int32* IndexPtr = PublicVariablesProxyMap.Find(Name);
+	if (IndexPtr == nullptr)
+	{
+		return TOptional<uint8>();
+	}
+
+	FRWScopeLock Lock(PublicVariablesLock, SLT_ReadOnly);
+
+	TConstArrayView<FPropertyBagPropertyDesc> ProxyDescs = PublicVariablesProxy.Data.GetPropertyBagStruct()->GetPropertyDescs();
+	const FPropertyBagPropertyDesc& PropertyDesc = ProxyDescs[*IndexPtr];
+
+	TValueOrError<uint8, EPropertyBagResult> Result = PublicVariablesProxy.Data.GetValueByte(PropertyDesc);
+	return Result.HasValue() ? TOptional<uint8>(Result.GetValue()) : TOptional<uint8>();
+}
+
+TOptional<int32> UAnimNextComponent::GetVariableInt32(const FName Name) const
+{
+	const int32* IndexPtr = PublicVariablesProxyMap.Find(Name);
+	if (IndexPtr == nullptr)
+	{
+		return TOptional<int32>();
+	}
+
+	FRWScopeLock Lock(PublicVariablesLock, SLT_ReadOnly);
+
+	TConstArrayView<FPropertyBagPropertyDesc> ProxyDescs = PublicVariablesProxy.Data.GetPropertyBagStruct()->GetPropertyDescs();
+	const FPropertyBagPropertyDesc& PropertyDesc = ProxyDescs[*IndexPtr];
+
+	TValueOrError<int32, EPropertyBagResult> Result = PublicVariablesProxy.Data.GetValueInt32(PropertyDesc);
+	return Result.HasValue() ? TOptional<int32>(Result.GetValue()) : TOptional<int32>();
+}
+
+TOptional<uint32> UAnimNextComponent::GetVariableUInt32(const FName Name) const
+{
+	const int32* IndexPtr = PublicVariablesProxyMap.Find(Name);
+	if (IndexPtr == nullptr)
+	{
+		return TOptional<uint32>();
+	}
+
+	FRWScopeLock Lock(PublicVariablesLock, SLT_ReadOnly);
+
+	TConstArrayView<FPropertyBagPropertyDesc> ProxyDescs = PublicVariablesProxy.Data.GetPropertyBagStruct()->GetPropertyDescs();
+	const FPropertyBagPropertyDesc& PropertyDesc = ProxyDescs[*IndexPtr];
+
+	TValueOrError<uint32, EPropertyBagResult> Result = PublicVariablesProxy.Data.GetValueUInt32(PropertyDesc);
+	return Result.HasValue() ? TOptional<uint32>(Result.GetValue()) : TOptional<uint32>();
+}
+
+TOptional<int64> UAnimNextComponent::GetVariableInt64(const FName Name) const
+{
+	const int32* IndexPtr = PublicVariablesProxyMap.Find(Name);
+	if (IndexPtr == nullptr)
+	{
+		return TOptional<int64>();
+	}
+
+	FRWScopeLock Lock(PublicVariablesLock, SLT_ReadOnly);
+
+	TConstArrayView<FPropertyBagPropertyDesc> ProxyDescs = PublicVariablesProxy.Data.GetPropertyBagStruct()->GetPropertyDescs();
+	const FPropertyBagPropertyDesc& PropertyDesc = ProxyDescs[*IndexPtr];
+
+	TValueOrError<int64, EPropertyBagResult> Result = PublicVariablesProxy.Data.GetValueInt64(PropertyDesc);
+	return Result.HasValue() ? TOptional<int64>(Result.GetValue()) : TOptional<int64>();
+}
+
+TOptional<uint64> UAnimNextComponent::GetVariableUInt64(const FName Name) const
+{
+	const int32* IndexPtr = PublicVariablesProxyMap.Find(Name);
+	if (IndexPtr == nullptr)
+	{
+		return TOptional<uint64>();
+	}
+
+	FRWScopeLock Lock(PublicVariablesLock, SLT_ReadOnly);
+
+	TConstArrayView<FPropertyBagPropertyDesc> ProxyDescs = PublicVariablesProxy.Data.GetPropertyBagStruct()->GetPropertyDescs();
+	const FPropertyBagPropertyDesc& PropertyDesc = ProxyDescs[*IndexPtr];
+
+	TValueOrError<uint64, EPropertyBagResult> Result = PublicVariablesProxy.Data.GetValueUInt64(PropertyDesc);
+	return Result.HasValue() ? TOptional<uint64>(Result.GetValue()) : TOptional<uint64>();
+}
+
+TOptional<float> UAnimNextComponent::GetVariableFloat(const FName Name) const
+{
+	const int32* IndexPtr = PublicVariablesProxyMap.Find(Name);
+	if (IndexPtr == nullptr)
+	{
+		return TOptional<float>();
+	}
+
+	FRWScopeLock Lock(PublicVariablesLock, SLT_ReadOnly);
+
+	TConstArrayView<FPropertyBagPropertyDesc> ProxyDescs = PublicVariablesProxy.Data.GetPropertyBagStruct()->GetPropertyDescs();
+	const FPropertyBagPropertyDesc& PropertyDesc = ProxyDescs[*IndexPtr];
+
+	TValueOrError<float, EPropertyBagResult> Result = PublicVariablesProxy.Data.GetValueFloat(PropertyDesc);
+	return Result.HasValue() ? TOptional<float>(Result.GetValue()) : TOptional<float>();
+}
+
+TOptional<double> UAnimNextComponent::GetVariableDouble(const FName Name) const
+{
+	const int32* IndexPtr = PublicVariablesProxyMap.Find(Name);
+	if (IndexPtr == nullptr)
+	{
+		return TOptional<double>();
+	}
+
+	FRWScopeLock Lock(PublicVariablesLock, SLT_ReadOnly);
+
+	TConstArrayView<FPropertyBagPropertyDesc> ProxyDescs = PublicVariablesProxy.Data.GetPropertyBagStruct()->GetPropertyDescs();
+	const FPropertyBagPropertyDesc& PropertyDesc = ProxyDescs[*IndexPtr];
+
+	TValueOrError<double, EPropertyBagResult> Result = PublicVariablesProxy.Data.GetValueDouble(PropertyDesc);
+	return Result.HasValue() ? TOptional<double>(Result.GetValue()) : TOptional<double>();
+}
+
+TOptional<FName> UAnimNextComponent::GetVariableName(const FName Name) const
+{
+	const int32* IndexPtr = PublicVariablesProxyMap.Find(Name);
+	if (IndexPtr == nullptr)
+	{
+		return TOptional<FName>();
+	}
+
+	FRWScopeLock Lock(PublicVariablesLock, SLT_ReadOnly);
+
+	TConstArrayView<FPropertyBagPropertyDesc> ProxyDescs = PublicVariablesProxy.Data.GetPropertyBagStruct()->GetPropertyDescs();
+	const FPropertyBagPropertyDesc& PropertyDesc = ProxyDescs[*IndexPtr];
+
+	TValueOrError<FName, EPropertyBagResult> Result = PublicVariablesProxy.Data.GetValueName(PropertyDesc);
+	return Result.HasValue() ? TOptional<FName>(Result.GetValue()) : TOptional<FName>();
+}
+
+TOptional<FString> UAnimNextComponent::GetVariableString(const FName Name) const
+{
+	const int32* IndexPtr = PublicVariablesProxyMap.Find(Name);
+	if (IndexPtr == nullptr)
+	{
+		return TOptional<FString>();
+	}
+
+	FRWScopeLock Lock(PublicVariablesLock, SLT_ReadOnly);
+
+	TConstArrayView<FPropertyBagPropertyDesc> ProxyDescs = PublicVariablesProxy.Data.GetPropertyBagStruct()->GetPropertyDescs();
+	const FPropertyBagPropertyDesc& PropertyDesc = ProxyDescs[*IndexPtr];
+
+	TValueOrError<FString, EPropertyBagResult> Result = PublicVariablesProxy.Data.GetValueString(PropertyDesc);
+	return Result.HasValue() ? TOptional<FString>(Result.GetValue()) : TOptional<FString>();
+}
+
+TOptional<uint8> UAnimNextComponent::GetVariableEnum(const FName Name, const UEnum* RequestedEnum) const
+{
+	const int32* IndexPtr = PublicVariablesProxyMap.Find(Name);
+	if (IndexPtr == nullptr)
+	{
+		return TOptional<uint8>();
+	}
+
+	FRWScopeLock Lock(PublicVariablesLock, SLT_ReadOnly);
+
+	TConstArrayView<FPropertyBagPropertyDesc> ProxyDescs = PublicVariablesProxy.Data.GetPropertyBagStruct()->GetPropertyDescs();
+	const FPropertyBagPropertyDesc& PropertyDesc = ProxyDescs[*IndexPtr];
+
+	TValueOrError<uint8, EPropertyBagResult> Result = PublicVariablesProxy.Data.GetValueEnum(PropertyDesc, RequestedEnum);
+	return Result.HasValue() ? TOptional<uint8>(Result.GetValue()) : TOptional<uint8>();
+}
+
+TOptional<FStructView> UAnimNextComponent::GetVariableStruct(const FName Name, const UScriptStruct* RequestedStruct) const
+{
+	const int32* IndexPtr = PublicVariablesProxyMap.Find(Name);
+	if (IndexPtr == nullptr)
+	{
+		return TOptional<FStructView>();
+	}
+
+	FRWScopeLock Lock(PublicVariablesLock, SLT_ReadOnly);
+
+	TConstArrayView<FPropertyBagPropertyDesc> ProxyDescs = PublicVariablesProxy.Data.GetPropertyBagStruct()->GetPropertyDescs();
+	const FPropertyBagPropertyDesc& PropertyDesc = ProxyDescs[*IndexPtr];
+
+	TValueOrError<FStructView, EPropertyBagResult> Result = PublicVariablesProxy.Data.GetValueStruct(PropertyDesc, RequestedStruct);
+	return Result.HasValue() ? TOptional<FStructView>(Result.GetValue()) : TOptional<FStructView>();
+}
+
+TOptional<UObject*> UAnimNextComponent::GetVariableObject(const FName Name, const UClass* RequestedClass) const
+{
+	const int32* IndexPtr = PublicVariablesProxyMap.Find(Name);
+	if (IndexPtr == nullptr)
+	{
+		return TOptional<UObject*>();
+	}
+
+	FRWScopeLock Lock(PublicVariablesLock, SLT_ReadOnly);
+
+	TConstArrayView<FPropertyBagPropertyDesc> ProxyDescs = PublicVariablesProxy.Data.GetPropertyBagStruct()->GetPropertyDescs();
+	const FPropertyBagPropertyDesc& PropertyDesc = ProxyDescs[*IndexPtr];
+
+	TValueOrError<UObject*, EPropertyBagResult> Result = PublicVariablesProxy.Data.GetValueObject(PropertyDesc, RequestedClass);
+	return Result.HasValue() ? TOptional<UObject*>(Result.GetValue()) : TOptional<UObject*>();
+}
+
+TOptional<UClass*> UAnimNextComponent::GetVariableClass(const FName Name) const
+{
+	const int32* IndexPtr = PublicVariablesProxyMap.Find(Name);
+	if (IndexPtr == nullptr)
+	{
+		return TOptional<UClass*>();
+	}
+
+	FRWScopeLock Lock(PublicVariablesLock, SLT_ReadOnly);
+
+	TConstArrayView<FPropertyBagPropertyDesc> ProxyDescs = PublicVariablesProxy.Data.GetPropertyBagStruct()->GetPropertyDescs();
+	const FPropertyBagPropertyDesc& PropertyDesc = ProxyDescs[*IndexPtr];
+
+	TValueOrError<UClass*, EPropertyBagResult> Result = PublicVariablesProxy.Data.GetValueClass(PropertyDesc);
+	return Result.HasValue() ? TOptional<UClass*>(Result.GetValue()) : TOptional<UClass*>();
+}
+
+TOptional<FSoftObjectPath> UAnimNextComponent::GetVariableSoftPath(const FName Name) const
+{
+	const int32* IndexPtr = PublicVariablesProxyMap.Find(Name);
+	if (IndexPtr == nullptr)
+	{
+		return TOptional<FSoftObjectPath>();
+	}
+
+	FRWScopeLock Lock(PublicVariablesLock, SLT_ReadOnly);
+
+	TConstArrayView<FPropertyBagPropertyDesc> ProxyDescs = PublicVariablesProxy.Data.GetPropertyBagStruct()->GetPropertyDescs();
+	const FPropertyBagPropertyDesc& PropertyDesc = ProxyDescs[*IndexPtr];
+
+	TValueOrError<FSoftObjectPath, EPropertyBagResult> Result = PublicVariablesProxy.Data.GetValueSoftPath(PropertyDesc);
+	return Result.HasValue() ? TOptional<FSoftObjectPath>(Result.GetValue()) : TOptional<FSoftObjectPath>();
+}
+
+bool UAnimNextComponent::SetVariableBool(const FName Name, const bool bInValue)
+{
+	const int32* IndexPtr = PublicVariablesProxyMap.Find(Name);
+	if (IndexPtr == nullptr)
+	{
+		return false;
+	}
+
+	FRWScopeLock Lock(PublicVariablesLock, SLT_Write);
+
+	TConstArrayView<FPropertyBagPropertyDesc> ProxyDescs = PublicVariablesProxy.Data.GetPropertyBagStruct()->GetPropertyDescs();
+	const FPropertyBagPropertyDesc& PropertyDesc = ProxyDescs[*IndexPtr];
+
+	EPropertyBagResult Result = PublicVariablesProxy.Data.SetValueBool(PropertyDesc, bInValue);
+	if (Result == EPropertyBagResult::Success)
+	{
+		PublicVariablesProxy.DirtyFlags[*IndexPtr] = true;
+		PublicVariablesProxy.bIsDirty = true;
+	}
+
+	return Result == EPropertyBagResult::Success;
+}
+
+bool UAnimNextComponent::SetVariableByte(const FName Name, const uint8 InValue)
+{
+	const int32* IndexPtr = PublicVariablesProxyMap.Find(Name);
+	if (IndexPtr == nullptr)
+	{
+		return false;
+	}
+
+	FRWScopeLock Lock(PublicVariablesLock, SLT_Write);
+
+	TConstArrayView<FPropertyBagPropertyDesc> ProxyDescs = PublicVariablesProxy.Data.GetPropertyBagStruct()->GetPropertyDescs();
+	const FPropertyBagPropertyDesc& PropertyDesc = ProxyDescs[*IndexPtr];
+
+	EPropertyBagResult Result = PublicVariablesProxy.Data.SetValueByte(PropertyDesc, InValue);
+	if (Result == EPropertyBagResult::Success)
+	{
+		PublicVariablesProxy.DirtyFlags[*IndexPtr] = true;
+		PublicVariablesProxy.bIsDirty = true;
+	}
+
+	return Result == EPropertyBagResult::Success;
+}
+
+bool UAnimNextComponent::SetVariableInt32(const FName Name, const int32 InValue)
+{
+	const int32* IndexPtr = PublicVariablesProxyMap.Find(Name);
+	if (IndexPtr == nullptr)
+	{
+		return false;
+	}
+
+	FRWScopeLock Lock(PublicVariablesLock, SLT_Write);
+
+	TConstArrayView<FPropertyBagPropertyDesc> ProxyDescs = PublicVariablesProxy.Data.GetPropertyBagStruct()->GetPropertyDescs();
+	const FPropertyBagPropertyDesc& PropertyDesc = ProxyDescs[*IndexPtr];
+
+	EPropertyBagResult Result = PublicVariablesProxy.Data.SetValueInt32(PropertyDesc, InValue);
+	if (Result == EPropertyBagResult::Success)
+	{
+		PublicVariablesProxy.DirtyFlags[*IndexPtr] = true;
+		PublicVariablesProxy.bIsDirty = true;
+	}
+
+	return Result == EPropertyBagResult::Success;
+}
+
+bool UAnimNextComponent::SetVariableUInt32(const FName Name, const uint32 InValue)
+{
+	const int32* IndexPtr = PublicVariablesProxyMap.Find(Name);
+	if (IndexPtr == nullptr)
+	{
+		return false;
+	}
+
+	FRWScopeLock Lock(PublicVariablesLock, SLT_Write);
+
+	TConstArrayView<FPropertyBagPropertyDesc> ProxyDescs = PublicVariablesProxy.Data.GetPropertyBagStruct()->GetPropertyDescs();
+	const FPropertyBagPropertyDesc& PropertyDesc = ProxyDescs[*IndexPtr];
+
+	EPropertyBagResult Result = PublicVariablesProxy.Data.SetValueUInt32(PropertyDesc, InValue);
+	if (Result == EPropertyBagResult::Success)
+	{
+		PublicVariablesProxy.DirtyFlags[*IndexPtr] = true;
+		PublicVariablesProxy.bIsDirty = true;
+	}
+
+	return Result == EPropertyBagResult::Success;
+}
+
+bool UAnimNextComponent::SetVariableInt64(const FName Name, const int64 InValue)
+{
+	const int32* IndexPtr = PublicVariablesProxyMap.Find(Name);
+	if (IndexPtr == nullptr)
+	{
+		return false;
+	}
+
+	FRWScopeLock Lock(PublicVariablesLock, SLT_Write);
+
+	TConstArrayView<FPropertyBagPropertyDesc> ProxyDescs = PublicVariablesProxy.Data.GetPropertyBagStruct()->GetPropertyDescs();
+	const FPropertyBagPropertyDesc& PropertyDesc = ProxyDescs[*IndexPtr];
+
+	EPropertyBagResult Result = PublicVariablesProxy.Data.SetValueInt64(PropertyDesc, InValue);
+	if (Result == EPropertyBagResult::Success)
+	{
+		PublicVariablesProxy.DirtyFlags[*IndexPtr] = true;
+		PublicVariablesProxy.bIsDirty = true;
+	}
+
+	return Result == EPropertyBagResult::Success;
+}
+
+bool UAnimNextComponent::SetVariableUInt64(const FName Name, const uint64 InValue)
+{
+	const int32* IndexPtr = PublicVariablesProxyMap.Find(Name);
+	if (IndexPtr == nullptr)
+	{
+		return false;
+	}
+
+	FRWScopeLock Lock(PublicVariablesLock, SLT_Write);
+
+	TConstArrayView<FPropertyBagPropertyDesc> ProxyDescs = PublicVariablesProxy.Data.GetPropertyBagStruct()->GetPropertyDescs();
+	const FPropertyBagPropertyDesc& PropertyDesc = ProxyDescs[*IndexPtr];
+
+	EPropertyBagResult Result = PublicVariablesProxy.Data.SetValueUInt64(PropertyDesc, InValue);
+	if (Result == EPropertyBagResult::Success)
+	{
+		PublicVariablesProxy.DirtyFlags[*IndexPtr] = true;
+		PublicVariablesProxy.bIsDirty = true;
+	}
+
+	return Result == EPropertyBagResult::Success;
+}
+
+bool UAnimNextComponent::SetVariableFloat(const FName Name, const float InValue)
+{
+	const int32* IndexPtr = PublicVariablesProxyMap.Find(Name);
+	if (IndexPtr == nullptr)
+	{
+		return false;
+	}
+
+	FRWScopeLock Lock(PublicVariablesLock, SLT_Write);
+
+	TConstArrayView<FPropertyBagPropertyDesc> ProxyDescs = PublicVariablesProxy.Data.GetPropertyBagStruct()->GetPropertyDescs();
+	const FPropertyBagPropertyDesc& PropertyDesc = ProxyDescs[*IndexPtr];
+
+	EPropertyBagResult Result = PublicVariablesProxy.Data.SetValueFloat(PropertyDesc, InValue);
+	if (Result == EPropertyBagResult::Success)
+	{
+		PublicVariablesProxy.DirtyFlags[*IndexPtr] = true;
+		PublicVariablesProxy.bIsDirty = true;
+	}
+
+	return Result == EPropertyBagResult::Success;
+}
+
+bool UAnimNextComponent::SetVariableDouble(const FName Name, const double InValue)
+{
+	const int32* IndexPtr = PublicVariablesProxyMap.Find(Name);
+	if (IndexPtr == nullptr)
+	{
+		return false;
+	}
+
+	FRWScopeLock Lock(PublicVariablesLock, SLT_Write);
+
+	TConstArrayView<FPropertyBagPropertyDesc> ProxyDescs = PublicVariablesProxy.Data.GetPropertyBagStruct()->GetPropertyDescs();
+	const FPropertyBagPropertyDesc& PropertyDesc = ProxyDescs[*IndexPtr];
+
+	EPropertyBagResult Result = PublicVariablesProxy.Data.SetValueDouble(PropertyDesc, InValue);
+	if (Result == EPropertyBagResult::Success)
+	{
+		PublicVariablesProxy.DirtyFlags[*IndexPtr] = true;
+		PublicVariablesProxy.bIsDirty = true;
+	}
+
+	return Result == EPropertyBagResult::Success;
+}
+
+bool UAnimNextComponent::SetVariableName(const FName Name, const FName InValue)
+{
+	const int32* IndexPtr = PublicVariablesProxyMap.Find(Name);
+	if (IndexPtr == nullptr)
+	{
+		return false;
+	}
+
+	FRWScopeLock Lock(PublicVariablesLock, SLT_Write);
+
+	TConstArrayView<FPropertyBagPropertyDesc> ProxyDescs = PublicVariablesProxy.Data.GetPropertyBagStruct()->GetPropertyDescs();
+	const FPropertyBagPropertyDesc& PropertyDesc = ProxyDescs[*IndexPtr];
+
+	EPropertyBagResult Result = PublicVariablesProxy.Data.SetValueName(PropertyDesc, InValue);
+	if (Result == EPropertyBagResult::Success)
+	{
+		PublicVariablesProxy.DirtyFlags[*IndexPtr] = true;
+		PublicVariablesProxy.bIsDirty = true;
+	}
+
+	return Result == EPropertyBagResult::Success;
+}
+
+bool UAnimNextComponent::SetVariableString(const FName Name, const FString& InValue)
+{
+	const int32* IndexPtr = PublicVariablesProxyMap.Find(Name);
+	if (IndexPtr == nullptr)
+	{
+		return false;
+	}
+
+	FRWScopeLock Lock(PublicVariablesLock, SLT_Write);
+
+	TConstArrayView<FPropertyBagPropertyDesc> ProxyDescs = PublicVariablesProxy.Data.GetPropertyBagStruct()->GetPropertyDescs();
+	const FPropertyBagPropertyDesc& PropertyDesc = ProxyDescs[*IndexPtr];
+
+	EPropertyBagResult Result = PublicVariablesProxy.Data.SetValueString(PropertyDesc, InValue);
+	if (Result == EPropertyBagResult::Success)
+	{
+		PublicVariablesProxy.DirtyFlags[*IndexPtr] = true;
+		PublicVariablesProxy.bIsDirty = true;
+	}
+
+	return Result == EPropertyBagResult::Success;
+}
+
+bool UAnimNextComponent::SetVariableEnum(const FName Name, const uint8 InValue, const UEnum* Enum)
+{
+	const int32* IndexPtr = PublicVariablesProxyMap.Find(Name);
+	if (IndexPtr == nullptr)
+	{
+		return false;
+	}
+
+	FRWScopeLock Lock(PublicVariablesLock, SLT_Write);
+
+	TConstArrayView<FPropertyBagPropertyDesc> ProxyDescs = PublicVariablesProxy.Data.GetPropertyBagStruct()->GetPropertyDescs();
+	const FPropertyBagPropertyDesc& PropertyDesc = ProxyDescs[*IndexPtr];
+
+	EPropertyBagResult Result = PublicVariablesProxy.Data.SetValueEnum(PropertyDesc, InValue, Enum);
+	if (Result == EPropertyBagResult::Success)
+	{
+		PublicVariablesProxy.DirtyFlags[*IndexPtr] = true;
+		PublicVariablesProxy.bIsDirty = true;
+	}
+
+	return Result == EPropertyBagResult::Success;
+}
+
+bool UAnimNextComponent::SetVariableStruct(const FName Name, FConstStructView InValue)
+{
+	const int32* IndexPtr = PublicVariablesProxyMap.Find(Name);
+	if (IndexPtr == nullptr)
+	{
+		return false;
+	}
+
+	FRWScopeLock Lock(PublicVariablesLock, SLT_Write);
+
+	TConstArrayView<FPropertyBagPropertyDesc> ProxyDescs = PublicVariablesProxy.Data.GetPropertyBagStruct()->GetPropertyDescs();
+	const FPropertyBagPropertyDesc& PropertyDesc = ProxyDescs[*IndexPtr];
+
+	EPropertyBagResult Result = PublicVariablesProxy.Data.SetValueStruct(PropertyDesc, InValue);
+	if (Result == EPropertyBagResult::Success)
+	{
+		PublicVariablesProxy.DirtyFlags[*IndexPtr] = true;
+		PublicVariablesProxy.bIsDirty = true;
+	}
+
+	return Result == EPropertyBagResult::Success;
+}
+
+bool UAnimNextComponent::SetVariableObject(const FName Name, UObject* InValue)
+{
+	const int32* IndexPtr = PublicVariablesProxyMap.Find(Name);
+	if (IndexPtr == nullptr)
+	{
+		return false;
+	}
+
+	FRWScopeLock Lock(PublicVariablesLock, SLT_Write);
+
+	TConstArrayView<FPropertyBagPropertyDesc> ProxyDescs = PublicVariablesProxy.Data.GetPropertyBagStruct()->GetPropertyDescs();
+	const FPropertyBagPropertyDesc& PropertyDesc = ProxyDescs[*IndexPtr];
+
+	EPropertyBagResult Result = PublicVariablesProxy.Data.SetValueObject(PropertyDesc, InValue);
+	if (Result == EPropertyBagResult::Success)
+	{
+		PublicVariablesProxy.DirtyFlags[*IndexPtr] = true;
+		PublicVariablesProxy.bIsDirty = true;
+	}
+
+	return Result == EPropertyBagResult::Success;
+}
+
+bool UAnimNextComponent::SetVariableClass(const FName Name, UClass* InValue)
+{
+	const int32* IndexPtr = PublicVariablesProxyMap.Find(Name);
+	if (IndexPtr == nullptr)
+	{
+		return false;
+	}
+
+	FRWScopeLock Lock(PublicVariablesLock, SLT_Write);
+
+	TConstArrayView<FPropertyBagPropertyDesc> ProxyDescs = PublicVariablesProxy.Data.GetPropertyBagStruct()->GetPropertyDescs();
+	const FPropertyBagPropertyDesc& PropertyDesc = ProxyDescs[*IndexPtr];
+
+	EPropertyBagResult Result = PublicVariablesProxy.Data.SetValueClass(PropertyDesc, InValue);
+	if (Result == EPropertyBagResult::Success)
+	{
+		PublicVariablesProxy.DirtyFlags[*IndexPtr] = true;
+		PublicVariablesProxy.bIsDirty = true;
+	}
+
+	return Result == EPropertyBagResult::Success;
+}
+
+bool UAnimNextComponent::SetVariableSoftPath(const FName Name, const FSoftObjectPath& InValue)
+{
+	const int32* IndexPtr = PublicVariablesProxyMap.Find(Name);
+	if (IndexPtr == nullptr)
+	{
+		return false;
+	}
+
+	FRWScopeLock Lock(PublicVariablesLock, SLT_Write);
+
+	TConstArrayView<FPropertyBagPropertyDesc> ProxyDescs = PublicVariablesProxy.Data.GetPropertyBagStruct()->GetPropertyDescs();
+	const FPropertyBagPropertyDesc& PropertyDesc = ProxyDescs[*IndexPtr];
+
+	EPropertyBagResult Result = PublicVariablesProxy.Data.SetValueSoftPath(PropertyDesc, InValue);
+	if (Result == EPropertyBagResult::Success)
+	{
+		PublicVariablesProxy.DirtyFlags[*IndexPtr] = true;
+		PublicVariablesProxy.bIsDirty = true;
+	}
+
+	return Result == EPropertyBagResult::Success;
 }
 
 void UAnimNextComponent::SetEnabled(bool bEnabled)
diff --git a/Engine/Source/Runtime/CoreUObject/Private/StructUtils/PropertyBag.cpp b/Engine/Source/Runtime/CoreUObject/Private/StructUtils/PropertyBag.cpp
--- a/Engine/Source/Runtime/CoreUObject/Private/StructUtils/PropertyBag.cpp
+++ b/Engine/Source/Runtime/CoreUObject/Private/StructUtils/PropertyBag.cpp
@@ -1786,6 +1786,15 @@ const FPropertyBagPropertyDesc* FInstancedPropertyBag::FindPropertyDescByName(co
 	return nullptr;
 }
 
+bool FInstancedPropertyBag::OwnsPropertyDesc(const FPropertyBagPropertyDesc& Desc) const
+{
+	if (const UPropertyBag* BagStruct = GetPropertyBagStruct())
+	{
+		return BagStruct->OwnsPropertyDesc(Desc);
+	}
+	return false;
+}
+
 const void* FInstancedPropertyBag::GetValueAddress(const FPropertyBagPropertyDesc* Desc) const
 {
 	if (Desc == nullptr || !Value.IsValid())
@@ -2024,6 +2033,268 @@ TValueOrError<FString, EPropertyBagResult> FInstancedPropertyBag::GetValueSerial
 	return MakeValue(OutStringValue);
 }
 
+TValueOrError<bool, EPropertyBagResult> FInstancedPropertyBag::GetValueBool(const FPropertyBagPropertyDesc& Desc) const
+{
+	if (!OwnsPropertyDesc(Desc))
+	{
+		return MakeError(EPropertyBagResult::PropertyNotFound);
+	}
+
+	int64 ReturnValue = 0;
+	const EPropertyBagResult Result = UE::StructUtils::Private::GetPropertyAsInt64(&Desc, GetValueAddress(&Desc), ReturnValue);
+	if (Result != EPropertyBagResult::Success)
+	{
+		return MakeError(Result);
+	}
+	return MakeValue(ReturnValue != 0);
+}
+
+TValueOrError<uint8, EPropertyBagResult> FInstancedPropertyBag::GetValueByte(const FPropertyBagPropertyDesc& Desc) const
+{
+	if (!OwnsPropertyDesc(Desc))
+	{
+		return MakeError(EPropertyBagResult::PropertyNotFound);
+	}
+
+	uint64 ReturnValue = 0;
+	const EPropertyBagResult Result = UE::StructUtils::Private::GetPropertyAsUInt64(&Desc, GetValueAddress(&Desc), ReturnValue);
+	if (Result != EPropertyBagResult::Success)
+	{
+		return MakeError(Result);
+	}
+	return MakeValue(static_cast<uint8>(ReturnValue));
+}
+
+TValueOrError<int32, EPropertyBagResult> FInstancedPropertyBag::GetValueInt32(const FPropertyBagPropertyDesc& Desc) const
+{
+	if (!OwnsPropertyDesc(Desc))
+	{
+		return MakeError(EPropertyBagResult::PropertyNotFound);
+	}
+
+	int64 ReturnValue = 0;
+	const EPropertyBagResult Result = UE::StructUtils::Private::GetPropertyAsInt64(&Desc, GetValueAddress(&Desc), ReturnValue);
+	if (Result != EPropertyBagResult::Success)
+	{
+		return MakeError(Result);
+	}
+	return MakeValue(static_cast<int32>(ReturnValue));
+}
+
+TValueOrError<uint32, EPropertyBagResult> FInstancedPropertyBag::GetValueUInt32(const FPropertyBagPropertyDesc& Desc) const
+{
+	if (!OwnsPropertyDesc(Desc))
+	{
+		return MakeError(EPropertyBagResult::PropertyNotFound);
+	}
+
+	uint64 ReturnValue = 0;
+	const EPropertyBagResult Result = UE::StructUtils::Private::GetPropertyAsUInt64(&Desc, GetValueAddress(&Desc), ReturnValue);
+	if (Result != EPropertyBagResult::Success)
+	{
+		return MakeError(Result);
+	}
+	return MakeValue(static_cast<uint32>(ReturnValue));
+}
+
+TValueOrError<int64, EPropertyBagResult> FInstancedPropertyBag::GetValueInt64(const FPropertyBagPropertyDesc& Desc) const
+{
+	if (!OwnsPropertyDesc(Desc))
+	{
+		return MakeError(EPropertyBagResult::PropertyNotFound);
+	}
+
+	int64 ReturnValue = 0;
+	const EPropertyBagResult Result = UE::StructUtils::Private::GetPropertyAsInt64(&Desc, GetValueAddress(&Desc), ReturnValue);
+	if (Result != EPropertyBagResult::Success)
+	{
+		return MakeError(Result);
+	}
+	return MakeValue(ReturnValue);
+}
+
+TValueOrError<uint64, EPropertyBagResult> FInstancedPropertyBag::GetValueUInt64(const FPropertyBagPropertyDesc& Desc) const
+{
+	if (!OwnsPropertyDesc(Desc))
+	{
+		return MakeError(EPropertyBagResult::PropertyNotFound);
+	}
+
+	uint64 ReturnValue = 0;
+	const EPropertyBagResult Result = UE::StructUtils::Private::GetPropertyAsUInt64(&Desc, GetValueAddress(&Desc), ReturnValue);
+	if (Result != EPropertyBagResult::Success)
+	{
+		return MakeError(Result);
+	}
+	return MakeValue((uint64)ReturnValue);
+}
+
+TValueOrError<float, EPropertyBagResult> FInstancedPropertyBag::GetValueFloat(const FPropertyBagPropertyDesc& Desc) const
+{
+	if (!OwnsPropertyDesc(Desc))
+	{
+		return MakeError(EPropertyBagResult::PropertyNotFound);
+	}
+
+	double ReturnValue = 0;
+	const EPropertyBagResult Result = UE::StructUtils::Private::GetPropertyAsDouble(&Desc, GetValueAddress(&Desc), ReturnValue);
+	if (Result != EPropertyBagResult::Success)
+	{
+		return MakeError(Result);
+	}
+	return MakeValue(static_cast<float>(ReturnValue));
+}
+
+TValueOrError<double, EPropertyBagResult> FInstancedPropertyBag::GetValueDouble(const FPropertyBagPropertyDesc& Desc) const
+{
+	if (!OwnsPropertyDesc(Desc))
+	{
+		return MakeError(EPropertyBagResult::PropertyNotFound);
+	}
+
+	double ReturnValue = 0;
+	const EPropertyBagResult Result = UE::StructUtils::Private::GetPropertyAsDouble(&Desc, GetValueAddress(&Desc), ReturnValue);
+	if (Result != EPropertyBagResult::Success)
+	{
+		return MakeError(Result);
+	}
+	return MakeValue(ReturnValue);
+}
+
+TValueOrError<FName, EPropertyBagResult> FInstancedPropertyBag::GetValueName(const FPropertyBagPropertyDesc& Desc) const
+{
+	if (!OwnsPropertyDesc(Desc))
+	{
+		return MakeError(EPropertyBagResult::PropertyNotFound);
+	}
+
+	FName ReturnValue;
+	const EPropertyBagResult Result = UE::StructUtils::Private::GetPropertyValue<FName, FNameProperty>(&Desc, GetValueAddress(&Desc), ReturnValue);
+	if (Result != EPropertyBagResult::Success)
+	{
+		return MakeError(Result);
+	}
+	return MakeValue(ReturnValue);
+}
+
+TValueOrError<FString, EPropertyBagResult> FInstancedPropertyBag::GetValueString(const FPropertyBagPropertyDesc& Desc) const
+{
+	if (!OwnsPropertyDesc(Desc))
+	{
+		return MakeError(EPropertyBagResult::PropertyNotFound);
+	}
+
+	FString ReturnValue;
+	const EPropertyBagResult Result = UE::StructUtils::Private::GetPropertyValue<FString, FStrProperty>(&Desc, GetValueAddress(&Desc), ReturnValue);
+	if (Result != EPropertyBagResult::Success)
+	{
+		return MakeError(Result);
+	}
+	return MakeValue(ReturnValue);
+}
+
+TValueOrError<FText, EPropertyBagResult> FInstancedPropertyBag::GetValueText(const FPropertyBagPropertyDesc& Desc) const
+{
+	if (!OwnsPropertyDesc(Desc))
+	{
+		return MakeError(EPropertyBagResult::PropertyNotFound);
+	}
+
+	FText ReturnValue;
+	const EPropertyBagResult Result = UE::StructUtils::Private::GetPropertyValue<FText, FTextProperty>(&Desc, GetValueAddress(&Desc), ReturnValue);
+	if (Result != EPropertyBagResult::Success)
+	{
+		return MakeError(Result);
+	}
+	return MakeValue(ReturnValue);
+}
+
+TValueOrError<uint8, EPropertyBagResult> FInstancedPropertyBag::GetValueEnum(const FPropertyBagPropertyDesc& Desc, const UEnum* RequestedEnum) const
+{
+	if (!OwnsPropertyDesc(Desc))
+	{
+		return MakeError(EPropertyBagResult::PropertyNotFound);
+	}
+
+	uint8 ReturnValue;
+	const EPropertyBagResult Result = UE::StructUtils::Private::GetPropertyValueAsEnum(&Desc, RequestedEnum, GetValueAddress(&Desc), ReturnValue);
+	if (Result != EPropertyBagResult::Success)
+	{
+		return MakeError(Result);
+	}
+	return MakeValue(ReturnValue);
+}
+
+TValueOrError<FStructView, EPropertyBagResult> FInstancedPropertyBag::GetValueStruct(const FPropertyBagPropertyDesc& Desc, const UScriptStruct* RequestedStruct) const
+{
+	if (!OwnsPropertyDesc(Desc))
+	{
+		return MakeError(EPropertyBagResult::PropertyNotFound);
+	}
+
+	FStructView ReturnValue;
+	const EPropertyBagResult Result = UE::StructUtils::Private::GetPropertyValueAsStruct(&Desc, RequestedStruct, GetValueAddress(&Desc), ReturnValue);
+	if (Result != EPropertyBagResult::Success)
+	{
+		return MakeError(Result);
+	}
+	return MakeValue(ReturnValue);
+}
+
+TValueOrError<UObject*, EPropertyBagResult> FInstancedPropertyBag::GetValueObject(const FPropertyBagPropertyDesc& Desc, const UClass* RequestedClass) const
+{
+	if (!OwnsPropertyDesc(Desc))
+	{
+		return MakeError(EPropertyBagResult::PropertyNotFound);
+	}
+
+	UObject* ReturnValue = nullptr;
+	const EPropertyBagResult Result = UE::StructUtils::Private::GetPropertyValueAsObject(&Desc, RequestedClass, GetValueAddress(&Desc), ReturnValue);
+	if (Result != EPropertyBagResult::Success)
+	{
+		return MakeError(Result);
+	}
+	return MakeValue(ReturnValue);
+}
+
+TValueOrError<UClass*, EPropertyBagResult> FInstancedPropertyBag::GetValueClass(const FPropertyBagPropertyDesc& Desc) const
+{
+	if (!OwnsPropertyDesc(Desc))
+	{
+		return MakeError(EPropertyBagResult::PropertyNotFound);
+	}
+
+	UObject* ReturnValue = nullptr;
+	const EPropertyBagResult Result = UE::StructUtils::Private::GetPropertyValueAsObject(&Desc, nullptr, GetValueAddress(&Desc), ReturnValue);
+	if (Result != EPropertyBagResult::Success)
+	{
+		return MakeError(Result);
+	}
+	UClass* Class = Cast<UClass>(ReturnValue);
+	if (Class == nullptr && ReturnValue != nullptr)
+	{
+		return MakeError(EPropertyBagResult::TypeMismatch);
+	}
+	return MakeValue(Class);
+}
+
+TValueOrError<FSoftObjectPath, EPropertyBagResult> FInstancedPropertyBag::GetValueSoftPath(const FPropertyBagPropertyDesc& Desc) const
+{
+	if (!OwnsPropertyDesc(Desc))
+	{
+		return MakeError(EPropertyBagResult::PropertyNotFound);
+	}
+
+	FSoftObjectPath ReturnValue;
+	const EPropertyBagResult Result = UE::StructUtils::Private::GetPropertyValueAsSoftPath(&Desc, GetValueAddress(&Desc), ReturnValue);
+	if (Result != EPropertyBagResult::Success)
+	{
+		return MakeError(Result);
+	}
+
+	return MakeValue(std::move(ReturnValue));
+}
+
 EPropertyBagResult FInstancedPropertyBag::SetValueBool(const FName Name, const bool bInValue)
 {
 	const FPropertyBagPropertyDesc* Desc = FindPropertyDescByName(Name);
@@ -2045,7 +2316,7 @@ EPropertyBagResult FInstancedPropertyBag::SetValueInt32(const FName Name, const
 EPropertyBagResult FInstancedPropertyBag::SetValueUInt32(const FName Name, const uint32 InValue)
 {
 	const FPropertyBagPropertyDesc* Desc = FindPropertyDescByName(Name);
-	return UE::StructUtils::Private::SetPropertyFromUInt64(Desc, GetMutableValueAddress(Desc), (int64)InValue);
+	return UE::StructUtils::Private::SetPropertyFromUInt64(Desc, GetMutableValueAddress(Desc), (uint64)InValue);
 }
 
 EPropertyBagResult FInstancedPropertyBag::SetValueInt64(const FName Name, const int64 InValue)
@@ -2057,7 +2328,7 @@ EPropertyBagResult FInstancedPropertyBag::SetValueInt64(const FName Name, const
 EPropertyBagResult FInstancedPropertyBag::SetValueUInt64(const FName Name, const uint64 InValue)
 {
 	const FPropertyBagPropertyDesc* Desc = FindPropertyDescByName(Name);
-	return UE::StructUtils::Private::SetPropertyFromUInt64(Desc, GetMutableValueAddress(Desc), static_cast<int64>(InValue));
+	return UE::StructUtils::Private::SetPropertyFromUInt64(Desc, GetMutableValueAddress(Desc), InValue);
 }
 
 EPropertyBagResult FInstancedPropertyBag::SetValueFloat(const FName Name, const float InValue)
@@ -2164,6 +2435,150 @@ EPropertyBagResult FInstancedPropertyBag::SetValue(const FName Name, const FProp
 	return EPropertyBagResult::Success;
 }
 
+EPropertyBagResult FInstancedPropertyBag::SetValueBool(const FPropertyBagPropertyDesc& Desc, const bool bInValue)
+{
+	if (!OwnsPropertyDesc(Desc))
+	{
+		return EPropertyBagResult::PropertyNotFound;
+	}
+	return UE::StructUtils::Private::SetPropertyFromInt64(&Desc, GetMutableValueAddress(&Desc), bInValue ? 1 : 0);
+}
+
+EPropertyBagResult FInstancedPropertyBag::SetValueByte(const FPropertyBagPropertyDesc& Desc, const uint8 InValue)
+{
+	if (!OwnsPropertyDesc(Desc))
+	{
+		return EPropertyBagResult::PropertyNotFound;
+	}
+	return UE::StructUtils::Private::SetPropertyFromUInt64(&Desc, GetMutableValueAddress(&Desc), InValue);
+}
+
+EPropertyBagResult FInstancedPropertyBag::SetValueInt32(const FPropertyBagPropertyDesc& Desc, const int32 InValue)
+{
+	if (!OwnsPropertyDesc(Desc))
+	{
+		return EPropertyBagResult::PropertyNotFound;
+	}
+	return UE::StructUtils::Private::SetPropertyFromInt64(&Desc, GetMutableValueAddress(&Desc), InValue);
+}
+
+EPropertyBagResult FInstancedPropertyBag::SetValueUInt32(const FPropertyBagPropertyDesc& Desc, const uint32 InValue)
+{
+	if (!OwnsPropertyDesc(Desc))
+	{
+		return EPropertyBagResult::PropertyNotFound;
+	}
+	return UE::StructUtils::Private::SetPropertyFromUInt64(&Desc, GetMutableValueAddress(&Desc), (uint64)InValue);
+}
+
+EPropertyBagResult FInstancedPropertyBag::SetValueInt64(const FPropertyBagPropertyDesc& Desc, const int64 InValue)
+{
+	if (!OwnsPropertyDesc(Desc))
+	{
+		return EPropertyBagResult::PropertyNotFound;
+	}
+	return UE::StructUtils::Private::SetPropertyFromInt64(&Desc, GetMutableValueAddress(&Desc), InValue);
+}
+
+EPropertyBagResult FInstancedPropertyBag::SetValueUInt64(const FPropertyBagPropertyDesc& Desc, const uint64 InValue)
+{
+	if (!OwnsPropertyDesc(Desc))
+	{
+		return EPropertyBagResult::PropertyNotFound;
+	}
+	return UE::StructUtils::Private::SetPropertyFromUInt64(&Desc, GetMutableValueAddress(&Desc), InValue);
+}
+
+EPropertyBagResult FInstancedPropertyBag::SetValueFloat(const FPropertyBagPropertyDesc& Desc, const float InValue)
+{
+	if (!OwnsPropertyDesc(Desc))
+	{
+		return EPropertyBagResult::PropertyNotFound;
+	}
+	return UE::StructUtils::Private::SetPropertyFromDouble(&Desc, GetMutableValueAddress(&Desc), InValue);
+}
+
+EPropertyBagResult FInstancedPropertyBag::SetValueDouble(const FPropertyBagPropertyDesc& Desc, const double InValue)
+{
+	if (!OwnsPropertyDesc(Desc))
+	{
+		return EPropertyBagResult::PropertyNotFound;
+	}
+	return UE::StructUtils::Private::SetPropertyFromDouble(&Desc, GetMutableValueAddress(&Desc), InValue);
+}
+
+EPropertyBagResult FInstancedPropertyBag::SetValueName(const FPropertyBagPropertyDesc& Desc, const FName InValue)
+{
+	if (!OwnsPropertyDesc(Desc))
+	{
+		return EPropertyBagResult::PropertyNotFound;
+	}
+	return UE::StructUtils::Private::SetPropertyValue<FName, FNameProperty>(&Desc, GetMutableValueAddress(&Desc), InValue);
+}
+
+EPropertyBagResult FInstancedPropertyBag::SetValueString(const FPropertyBagPropertyDesc& Desc, const FString& InValue)
+{
+	if (!OwnsPropertyDesc(Desc))
+	{
+		return EPropertyBagResult::PropertyNotFound;
+	}
+	return UE::StructUtils::Private::SetPropertyValue<FString, FStrProperty>(&Desc, GetMutableValueAddress(&Desc), InValue);
+}
+
+EPropertyBagResult FInstancedPropertyBag::SetValueText(const FPropertyBagPropertyDesc& Desc, const FText& InValue)
+{
+	if (!OwnsPropertyDesc(Desc))
+	{
+		return EPropertyBagResult::PropertyNotFound;
+	}
+	return UE::StructUtils::Private::SetPropertyValue<FText, FTextProperty>(&Desc, GetMutableValueAddress(&Desc), InValue);
+}
+
+EPropertyBagResult FInstancedPropertyBag::SetValueEnum(const FPropertyBagPropertyDesc& Desc, const uint8 InValue, const UEnum* Enum)
+{
+	if (!OwnsPropertyDesc(Desc))
+	{
+		return EPropertyBagResult::PropertyNotFound;
+	}
+	return UE::StructUtils::Private::SetPropertyValueAsEnum(&Desc, GetMutableValueAddress(&Desc), InValue, Enum);
+}
+
+EPropertyBagResult FInstancedPropertyBag::SetValueStruct(const FPropertyBagPropertyDesc& Desc, FConstStructView InValue)
+{
+	if (!OwnsPropertyDesc(Desc))
+	{
+		return EPropertyBagResult::PropertyNotFound;
+	}
+	return UE::StructUtils::Private::SetPropertyValueAsStruct(&Desc, GetMutableValueAddress(&Desc), InValue);
+}
+
+EPropertyBagResult FInstancedPropertyBag::SetValueObject(const FPropertyBagPropertyDesc& Desc, UObject* InValue)
+{
+	if (!OwnsPropertyDesc(Desc))
+	{
+		return EPropertyBagResult::PropertyNotFound;
+	}
+	return UE::StructUtils::Private::SetPropertyValueAsObject(&Desc, GetMutableValueAddress(&Desc), InValue);
+}
+
+EPropertyBagResult FInstancedPropertyBag::SetValueClass(const FPropertyBagPropertyDesc& Desc, UClass* InValue)
+{
+	if (!OwnsPropertyDesc(Desc))
+	{
+		return EPropertyBagResult::PropertyNotFound;
+	}
+	return UE::StructUtils::Private::SetPropertyValueAsObject(&Desc, GetMutableValueAddress(&Desc), InValue);
+}
+
+EPropertyBagResult FInstancedPropertyBag::SetValueSoftPath(const FPropertyBagPropertyDesc& Desc, const FSoftObjectPath& InValue)
+{
+	if (!OwnsPropertyDesc(Desc))
+	{
+		return EPropertyBagResult::PropertyNotFound;
+	}
+	return UE::StructUtils::Private::SetPropertyValueAsSoftPath(&Desc, GetMutableValueAddress(&Desc), InValue);
+}
+
 TValueOrError<FPropertyBagArrayRef, EPropertyBagResult> FInstancedPropertyBag::GetMutableArrayRef(const FName Name)
 {
 	const FPropertyBagPropertyDesc* Desc = FindPropertyDescByName(Name);
@@ -3166,4 +3581,11 @@ const FPropertyBagPropertyDesc* UPropertyBag::FindPropertyDescByPropertyName(con
 const FPropertyBagPropertyDesc* UPropertyBag::FindPropertyDescByProperty(const FProperty* Property) const
 {
 	return PropertyDescs.FindByPredicate([&Property](const FPropertyBagPropertyDesc& Desc) { return Desc.CachedProperty && Desc.CachedProperty == Property; });
-}
\ No newline at end of file
+}
+
+bool UPropertyBag::OwnsPropertyDesc(const FPropertyBagPropertyDesc& Desc) const
+{
+	const FPropertyBagPropertyDesc* StartDesc = PropertyDescs.GetData();
+	const FPropertyBagPropertyDesc* EndDesc = StartDesc + PropertyDescs.Num();
+	return &Desc >= StartDesc && &Desc < EndDesc;
+}
diff --git a/Engine/Source/Runtime/CoreUObject/Public/StructUtils/PropertyBag.h b/Engine/Source/Runtime/CoreUObject/Public/StructUtils/PropertyBag.h
--- a/Engine/Source/Runtime/CoreUObject/Public/StructUtils/PropertyBag.h
+++ b/Engine/Source/Runtime/CoreUObject/Public/StructUtils/PropertyBag.h
@@ -471,6 +471,9 @@ struct COREUOBJECT_API FInstancedPropertyBag
 	/** Returns property descriptor by specified ID. */
 	const FPropertyBagPropertyDesc* FindPropertyDescByName(const FName Name) const;
 
+	/** Returns true if we own the supplied property description, false otherwise. */
+	bool OwnsPropertyDesc(const FPropertyBagPropertyDesc& Desc) const;
+
 	/** @return const view to the struct that holds the values. NOTE: The returned value/view cannot be serialized, use this to access the struct only temporarily. */
 	FConstStructView GetValue() const { return Value; };
 
@@ -554,6 +557,75 @@ struct COREUOBJECT_API FInstancedPropertyBag
 		return MakeError(EPropertyBagResult::TypeMismatch);
 	}
 
+	TValueOrError<bool, EPropertyBagResult> GetValueBool(const FPropertyBagPropertyDesc& Desc) const;
+	TValueOrError<uint8, EPropertyBagResult> GetValueByte(const FPropertyBagPropertyDesc& Desc) const;
+	TValueOrError<int32, EPropertyBagResult> GetValueInt32(const FPropertyBagPropertyDesc& Desc) const;
+	TValueOrError<uint32, EPropertyBagResult> GetValueUInt32(const FPropertyBagPropertyDesc& Desc) const;
+	TValueOrError<int64, EPropertyBagResult> GetValueInt64(const FPropertyBagPropertyDesc& Desc) const;
+	TValueOrError<uint64, EPropertyBagResult> GetValueUInt64(const FPropertyBagPropertyDesc& Desc) const;
+	TValueOrError<float, EPropertyBagResult> GetValueFloat(const FPropertyBagPropertyDesc& Desc) const;
+	TValueOrError<double, EPropertyBagResult> GetValueDouble(const FPropertyBagPropertyDesc& Desc) const;
+	TValueOrError<FName, EPropertyBagResult> GetValueName(const FPropertyBagPropertyDesc& Desc) const;
+	TValueOrError<FString, EPropertyBagResult> GetValueString(const FPropertyBagPropertyDesc& Desc) const;
+	TValueOrError<FText, EPropertyBagResult> GetValueText(const FPropertyBagPropertyDesc& Desc) const;
+	TValueOrError<uint8, EPropertyBagResult> GetValueEnum(const FPropertyBagPropertyDesc& Desc, const UEnum* RequestedEnum) const;
+	TValueOrError<FStructView, EPropertyBagResult> GetValueStruct(const FPropertyBagPropertyDesc& Desc, const UScriptStruct* RequestedStruct = nullptr) const;
+	TValueOrError<UObject*, EPropertyBagResult> GetValueObject(const FPropertyBagPropertyDesc& Desc, const UClass* RequestedClass = nullptr) const;
+	TValueOrError<UClass*, EPropertyBagResult> GetValueClass(const FPropertyBagPropertyDesc& Desc) const;
+	TValueOrError<FSoftObjectPath, EPropertyBagResult> GetValueSoftPath(const FPropertyBagPropertyDesc& Desc) const;
+
+	/** @return enum value of specified type. */
+	template <typename T>
+	TValueOrError<T, EPropertyBagResult> GetValueEnum(const FPropertyBagPropertyDesc& Desc) const
+	{
+		static_assert(TIsEnum<T>::Value, "Should only call this with enum types");
+
+		TValueOrError<uint8, EPropertyBagResult> Result = GetValueEnum(Desc, StaticEnum<T>());
+		if (!Result.IsValid())
+		{
+			return MakeError(Result.GetError());
+		}
+		return MakeValue(static_cast<T>(Result.GetValue()));
+	}
+
+	/** @return struct reference of specified type. */
+	template <typename T>
+	TValueOrError<T*, EPropertyBagResult> GetValueStruct(const FPropertyBagPropertyDesc& Desc) const
+	{
+		TValueOrError<FStructView, EPropertyBagResult> Result = GetValueStruct(Desc, TBaseStructure<T>::Get());
+		if (!Result.IsValid())
+		{
+			return MakeError(Result.GetError());
+		}
+		if (T* ValuePtr = Result.GetValue().GetPtr<T>())
+		{
+			return MakeValue(ValuePtr);
+		}
+		return MakeError(EPropertyBagResult::TypeMismatch);
+	}
+
+	/** @return object pointer value of specified type. */
+	template <typename T>
+	TValueOrError<T*, EPropertyBagResult> GetValueObject(const FPropertyBagPropertyDesc& Desc) const
+	{
+		static_assert(TIsDerivedFrom<T, UObject>::Value, "Should only call this with object types");
+
+		TValueOrError<UObject*, EPropertyBagResult> Result = GetValueObject(Desc, T::StaticClass());
+		if (!Result.IsValid())
+		{
+			return MakeError(Result.GetError());
+		}
+		if (Result.GetValue() == nullptr)
+		{
+			return MakeValue(nullptr);
+		}
+		if (T* Object = Cast<T>(Result.GetValue()))
+		{
+			return MakeValue(Object);
+		}
+		return MakeError(EPropertyBagResult::TypeMismatch);
+	}
+
 	/**
 	 * Value Setters. A property must exist in that bag before it can be set.  
 	 * Numeric types (bool, (u)int32, (u)int64, float, double) support type conversion.
@@ -610,6 +682,46 @@ struct COREUOBJECT_API FInstancedPropertyBag
 	 */
 	EPropertyBagResult SetValue(const FName Name, const FProperty* InSourceProperty, const void* InSourceContainerAddress);
 
+	EPropertyBagResult SetValueBool(const FPropertyBagPropertyDesc& Desc, const bool bInValue);
+	EPropertyBagResult SetValueByte(const FPropertyBagPropertyDesc& Desc, const uint8 InValue);
+	EPropertyBagResult SetValueInt32(const FPropertyBagPropertyDesc& Desc, const int32 InValue);
+	EPropertyBagResult SetValueUInt32(const FPropertyBagPropertyDesc& Desc, const uint32 InValue);
+	EPropertyBagResult SetValueInt64(const FPropertyBagPropertyDesc& Desc, const int64 InValue);
+	EPropertyBagResult SetValueUInt64(const FPropertyBagPropertyDesc& Desc, const uint64 InValue);
+	EPropertyBagResult SetValueFloat(const FPropertyBagPropertyDesc& Desc, const float InValue);
+	EPropertyBagResult SetValueDouble(const FPropertyBagPropertyDesc& Desc, const double InValue);
+	EPropertyBagResult SetValueName(const FPropertyBagPropertyDesc& Desc, const FName InValue);
+	EPropertyBagResult SetValueString(const FPropertyBagPropertyDesc& Desc, const FString& InValue);
+	EPropertyBagResult SetValueText(const FPropertyBagPropertyDesc& Desc, const FText& InValue);
+	EPropertyBagResult SetValueEnum(const FPropertyBagPropertyDesc& Desc, const uint8 InValue, const UEnum* Enum);
+	EPropertyBagResult SetValueStruct(const FPropertyBagPropertyDesc& Desc, FConstStructView InValue);
+	EPropertyBagResult SetValueObject(const FPropertyBagPropertyDesc& Desc, UObject* InValue);
+	EPropertyBagResult SetValueClass(const FPropertyBagPropertyDesc& Desc, UClass* InValue);
+	EPropertyBagResult SetValueSoftPath(const FPropertyBagPropertyDesc& Desc, const FSoftObjectPath& InValue);
+
+	/** Sets enum value specified type. */
+	template <typename T>
+	EPropertyBagResult SetValueEnum(const FPropertyBagPropertyDesc& Desc, const T InValue)
+	{
+		static_assert(TIsEnum<T>::Value, "Should only call this with enum types");
+		return SetValueEnum(Desc, static_cast<uint8>(InValue), StaticEnum<T>());
+	}
+
+	/** Sets struct value specified type. */
+	template <typename T>
+	EPropertyBagResult SetValueStruct(const FPropertyBagPropertyDesc& Desc, const T& InValue)
+	{
+		return SetValueStruct(Desc, FConstStructView::Make(InValue));
+	}
+
+	/** Sets object pointer value specified type. */
+	template <typename T>
+	EPropertyBagResult SetValueObject(const FPropertyBagPropertyDesc& Desc, T* InValue)
+	{
+		static_assert(TIsDerivedFrom<T, UObject>::Value, "Should only call this with object types");
+		return SetValueObject(Desc, (UObject*)InValue);
+	}
+
 	/**
 	 * Returns helper class to modify and access an array property.
 	 * Note: The array reference is not valid after the layout of the referenced property bag has changed!
@@ -1048,6 +1160,9 @@ public:
 	/** @return property description based on pointer to property. */
 	const FPropertyBagPropertyDesc* FindPropertyDescByProperty(const FProperty* Property) const;
 
+	/** @return true if we own the supplied property description, false otherwise. */
+	bool OwnsPropertyDesc(const FPropertyBagPropertyDesc& Desc) const;
+
 #if WITH_EDITOR
 	/** @return true if any of the properties on the bag has type of the specified user defined struct. */
 	bool ContainsUserDefinedStruct(const UUserDefinedStruct* UserDefinedStruct) const;
