Add missing ANIMNEXT_API

#tests PIE, nonunity

diff --git a/Engine/Plugins/Experimental/AnimNext/Source/AnimNext/Internal/Component/AnimNextComponent.h b/Engine/Plugins/Experimental/AnimNext/Source/AnimNext/Internal/Component/AnimNextComponent.h
--- a/Engine/Plugins/Experimental/AnimNext/Source/AnimNext/Internal/Component/AnimNextComponent.h
+++ b/Engine/Plugins/Experimental/AnimNext/Source/AnimNext/Internal/Component/AnimNextComponent.h
@@ -63,21 +63,21 @@ public:
 	 * @param    Name     The name of the variable to get
 	 * @return            The variable value on success, None otherwise
 	 */
-	TOptional<bool> GetVariableBool(const FName Name) const;
-	TOptional<uint8> GetVariableByte(const FName Name) const;
-	TOptional<int32> GetVariableInt32(const FName Name) const;
-	TOptional<uint32> GetVariableUInt32(const FName Name) const;
-	TOptional<int64> GetVariableInt64(const FName Name) const;
-	TOptional<uint64> GetVariableUInt64(const FName Name) const;
-	TOptional<float> GetVariableFloat(const FName Name) const;
-	TOptional<double> GetVariableDouble(const FName Name) const;
-	TOptional<FName> GetVariableName(const FName Name) const;
-	TOptional<FString> GetVariableString(const FName Name) const;
-	TOptional<uint8> GetVariableEnum(const FName Name, const UEnum* RequestedEnum) const;
-	TOptional<FStructView> GetVariableStruct(const FName Name, const UScriptStruct* RequestedStruct = nullptr) const;
-	TOptional<UObject*> GetVariableObject(const FName Name, const UClass* RequestedClass = nullptr) const;
-	TOptional<UClass*> GetVariableClass(const FName Name) const;
-	TOptional<FSoftObjectPath> GetVariableSoftPath(const FName Name) const;
+	ANIMNEXT_API TOptional<bool> GetVariableBool(const FName Name) const;
+	ANIMNEXT_API TOptional<uint8> GetVariableByte(const FName Name) const;
+	ANIMNEXT_API TOptional<int32> GetVariableInt32(const FName Name) const;
+	ANIMNEXT_API TOptional<uint32> GetVariableUInt32(const FName Name) const;
+	ANIMNEXT_API TOptional<int64> GetVariableInt64(const FName Name) const;
+	ANIMNEXT_API TOptional<uint64> GetVariableUInt64(const FName Name) const;
+	ANIMNEXT_API TOptional<float> GetVariableFloat(const FName Name) const;
+	ANIMNEXT_API TOptional<double> GetVariableDouble(const FName Name) const;
+	ANIMNEXT_API TOptional<FName> GetVariableName(const FName Name) const;
+	ANIMNEXT_API TOptional<FString> GetVariableString(const FName Name) const;
+	ANIMNEXT_API TOptional<uint8> GetVariableEnum(const FName Name, const UEnum* RequestedEnum) const;
+	ANIMNEXT_API TOptional<FStructView> GetVariableStruct(const FName Name, const UScriptStruct* RequestedStruct = nullptr) const;
+	ANIMNEXT_API TOptional<UObject*> GetVariableObject(const FName Name, const UClass* RequestedClass = nullptr) const;
+	ANIMNEXT_API TOptional<UClass*> GetVariableClass(const FName Name) const;
+	ANIMNEXT_API TOptional<FSoftObjectPath> GetVariableSoftPath(const FName Name) const;
 
 	/** @return enum value of specified type. */
 	template <typename T>
@@ -137,21 +137,21 @@ public:
 	 * @param    Value    The value to set the variable to
 	 * @return            True on success, false otherwise
 	 */
-	bool SetVariableBool(const FName Name, const bool bInValue);
-	bool SetVariableByte(const FName Name, const uint8 InValue);
-	bool SetVariableInt32(const FName Name, const int32 InValue);
-	bool SetVariableUInt32(const FName Name, const uint32 InValue);
-	bool SetVariableInt64(const FName Name, const int64 InValue);
-	bool SetVariableUInt64(const FName Name, const uint64 InValue);
-	bool SetVariableFloat(const FName Name, const float InValue);
-	bool SetVariableDouble(const FName Name, const double InValue);
-	bool SetVariableName(const FName Name, const FName InValue);
-	bool SetVariableString(const FName Name, const FString& InValue);
-	bool SetVariableEnum(const FName Name, const uint8 InValue, const UEnum* Enum);
-	bool SetVariableStruct(const FName Name, FConstStructView InValue);
-	bool SetVariableObject(const FName Name, UObject* InValue);
-	bool SetVariableClass(const FName Name, UClass* InValue);
-	bool SetVariableSoftPath(const FName Name, const FSoftObjectPath& InValue);
+	ANIMNEXT_API bool SetVariableBool(const FName Name, const bool bInValue);
+	ANIMNEXT_API bool SetVariableByte(const FName Name, const uint8 InValue);
+	ANIMNEXT_API bool SetVariableInt32(const FName Name, const int32 InValue);
+	ANIMNEXT_API bool SetVariableUInt32(const FName Name, const uint32 InValue);
+	ANIMNEXT_API bool SetVariableInt64(const FName Name, const int64 InValue);
+	ANIMNEXT_API bool SetVariableUInt64(const FName Name, const uint64 InValue);
+	ANIMNEXT_API bool SetVariableFloat(const FName Name, const float InValue);
+	ANIMNEXT_API bool SetVariableDouble(const FName Name, const double InValue);
+	ANIMNEXT_API bool SetVariableName(const FName Name, const FName InValue);
+	ANIMNEXT_API bool SetVariableString(const FName Name, const FString& InValue);
+	ANIMNEXT_API bool SetVariableEnum(const FName Name, const uint8 InValue, const UEnum* Enum);
+	ANIMNEXT_API bool SetVariableStruct(const FName Name, FConstStructView InValue);
+	ANIMNEXT_API bool SetVariableObject(const FName Name, UObject* InValue);
+	ANIMNEXT_API bool SetVariableClass(const FName Name, UClass* InValue);
+	ANIMNEXT_API bool SetVariableSoftPath(const FName Name, const FSoftObjectPath& InValue);
 
 	/** Sets enum variable with specified type. */
 	template <typename T>
