Running different LOD levels between running graphs as well as the visual mesh in the same AnimNext module

e.g. running facial animation on another LOD than the driving animation on the body. And also running animation on a different LOD than the visual representation/skeletal mesh component.

* Running animation graph with a lower-quality LOD level than the visual mesh now works. In this case we fill the missing transforms that the visual representation needs with reference pose transforms.
* Running multiple anim graphs within a module with different LOD levels now works:
- Source and target pose equal: Same as before, copying all transforms.
- Source LOD < target LOD: Source LOD is set to a higher-quality LOD level and contains more bone transforms than we need. Just copy the ones we actually need.
- Source LOD > target LOD: The source pose is missing transforms as it is set to a lower-quality LOD level, initialize the missing bone transforms with the reference pose transforms.

diff --git a/Engine/Plugins/Experimental/AnimNext/Source/AnimNext/Private/Graph/RigUnit_AnimNextWriteSkeletalMeshComponentPose.cpp b/Engine/Plugins/Experimental/AnimNext/Source/AnimNext/Private/Graph/RigUnit_AnimNextWriteSkeletalMeshComponentPose.cpp
--- a/Engine/Plugins/Experimental/AnimNext/Source/AnimNext/Private/Graph/RigUnit_AnimNextWriteSkeletalMeshComponentPose.cpp
+++ b/Engine/Plugins/Experimental/AnimNext/Source/AnimNext/Private/Graph/RigUnit_AnimNextWriteSkeletalMeshComponentPose.cpp
@@ -47,8 +47,29 @@ FRigUnit_AnimNextWriteSkeletalMeshComponentPose_Execute()
 
 	FMemMark MemMark(FMemStack::Get());
 
+	const int32 SkeletalMeshComponentLODLevel = SkeletalMeshComponent->GetPredictedLODLevel();
+	const int32 PoseLODLevel = Pose.LODPose.LODLevel;
+
 	TArray<FTransform, TMemStackAllocator<>> LocalSpaceTransforms;
-	LocalSpaceTransforms.SetNumUninitialized(SkeletalMesh->GetRefSkeleton().GetNum());
+
+	const FReferenceSkeleton& RefSkeleton = SkeletalMesh->GetRefSkeleton();
+	LocalSpaceTransforms.SetNumUninitialized(RefSkeleton.GetNum());
+
+	// Did we evaluate animation with lower quality than the visual mesh?
+	if (SkeletalMeshComponentLODLevel < PoseLODLevel)
+	{
+		// The given pose is missing transforms that the visual representation needs, initialize these with the reference pose.
+
+		const TArrayView<const FBoneIndexType> PoseBoneToMeshBoneIndexMap = RefPose.GetLODBoneIndexToMeshBoneIndexMap(SkeletalMeshComponentLODLevel);
+		const int32 NumPoseBones = Pose.LODPose.LocalTransforms.Num();
+		const int32 NumSkeletalMeshBones = PoseBoneToMeshBoneIndexMap.Num();
+
+		for (int32 PoseBoneIndex = NumPoseBones; PoseBoneIndex < NumSkeletalMeshBones; ++PoseBoneIndex)
+		{
+			const int32 MeshBoneIndex = PoseBoneToMeshBoneIndexMap[PoseBoneIndex];
+			LocalSpaceTransforms[MeshBoneIndex] = RefSkeleton.GetRefBonePose()[MeshBoneIndex];
+		}
+	}
 
 	// Map LOD pose into local-space scratch buffer
 	FGenerationTools::RemapPose(Pose.LODPose, LocalSpaceTransforms);
@@ -62,6 +83,6 @@ FRigUnit_AnimNextWriteSkeletalMeshComponentPose_Execute()
 	UE::Anim::FSkinnedMeshComponentExtensions::CompleteAndDispatch(
 		SkeletalMeshComponent,
 		RefPose.GetMeshBoneIndexToParentMeshBoneIndexMap(),
-		RefPose.GetLODBoneIndexToMeshBoneIndexMap(Pose.LODPose.LODLevel),
+		RefPose.GetLODBoneIndexToMeshBoneIndexMap(SkeletalMeshComponentLODLevel), // Based on the LOD level of the skeletal mesh component
 		LocalSpaceTransforms);
 }
diff --git a/Engine/Plugins/Experimental/AnimNext/Source/AnimNext/Private/TransformArrayOperations_AoS.cpp b/Engine/Plugins/Experimental/AnimNext/Source/AnimNext/Private/TransformArrayOperations_AoS.cpp
--- a/Engine/Plugins/Experimental/AnimNext/Source/AnimNext/Private/TransformArrayOperations_AoS.cpp
+++ b/Engine/Plugins/Experimental/AnimNext/Source/AnimNext/Private/TransformArrayOperations_AoS.cpp
@@ -8,9 +8,19 @@ namespace UE::AnimNext
 	void SetIdentity(const FTransformArrayAoSView& Dest, bool bIsAdditive)
 	{
 		const int32 NumTransforms = Dest.Num();
+		SetIdentity(Dest, bIsAdditive, /*StartIndex=*/0, NumTransforms);
+	}
+
+	void SetIdentity(const FTransformArrayAoSView& Dest, bool bIsAdditive, int32 StartIndex, int32 NumTransformsToSet)
+	{
 		const FTransform Identity = bIsAdditive ? TransformAdditiveIdentity : FTransform::Identity;
 
-		for (int32 TransformIndex = 0; TransformIndex < NumTransforms; ++TransformIndex)
+		const int EndIndex = StartIndex + NumTransformsToSet;
+		check((StartIndex + NumTransformsToSet) <= Dest.Num());
+		check(StartIndex >= 0 && StartIndex <= Dest.Num());
+		check(EndIndex <= Dest.Num());
+
+		for (int32 TransformIndex = StartIndex; TransformIndex < EndIndex; ++TransformIndex)
 		{
 			Dest[TransformIndex] = Identity;
 		}
diff --git a/Engine/Plugins/Experimental/AnimNext/Source/AnimNext/Private/TransformArrayOperations_SoA.cpp b/Engine/Plugins/Experimental/AnimNext/Source/AnimNext/Private/TransformArrayOperations_SoA.cpp
--- a/Engine/Plugins/Experimental/AnimNext/Source/AnimNext/Private/TransformArrayOperations_SoA.cpp
+++ b/Engine/Plugins/Experimental/AnimNext/Source/AnimNext/Private/TransformArrayOperations_SoA.cpp
@@ -11,6 +11,12 @@ DEFINE_STAT(STAT_AnimNext_BlendAccumulate_SoA);
 
 namespace UE::AnimNext
 {
+	void SetIdentity(const FTransformArraySoAView& Dest, bool bIsAdditive)
+	{
+		const int32 NumTransforms = Dest.Num();
+		SetIdentity(Dest, bIsAdditive, /*StartIndex=*/0, NumTransforms);
+	}
+
 	// A couple of performance notes:
 	//    * Types all use doubles
 	//    * Because we split transforms into 3 arrays, we pay 3x the range checks
@@ -41,9 +47,13 @@ namespace UE::AnimNext
 	//      The idea being to keep loops smaller to maximize overlap of instructions while waiting for long
 	//      latency instructions (cache miss in first batch, sqrt, div, etc).
 
-	void SetIdentity(const FTransformArraySoAView& Dest, bool bIsAdditive)
+	void SetIdentity(const FTransformArraySoAView& Dest, bool bIsAdditive, int32 StartIndex, int32 NumTransformsToSet)
 	{
-		const int32 NumTransforms = Dest.Num();
+		const int EndIndex = StartIndex + NumTransformsToSet;
+
+		check((StartIndex + NumTransformsToSet) <= Dest.Num());
+		check(StartIndex >= 0 && StartIndex <= Dest.Num());
+		check(EndIndex <= Dest.Num());
 
 		const FTransform& Identity = bIsAdditive ? TransformAdditiveIdentity : FTransform::Identity;
 
@@ -51,7 +61,7 @@ namespace UE::AnimNext
 		const FQuat IdentityRotation = Identity.GetRotation();
 		const FVector IdentityScale3D = Identity.GetScale3D();
 
-		for (int32 TransformIndex = 0; TransformIndex < NumTransforms; ++TransformIndex)
+		for (int32 TransformIndex = StartIndex; TransformIndex < EndIndex; ++TransformIndex)
 		{
 			Dest.Translations[TransformIndex] = IdentityTranslation;
 			Dest.Rotations[TransformIndex] = IdentityRotation;
@@ -66,7 +76,7 @@ namespace UE::AnimNext
 		const int32 NumTransforms = Dest.Num();
 		const int32 EndIndex = NumToCopy >= 0 ? (StartIndex + NumToCopy) : NumTransforms;
 
-		check(Source.Num() >= NumTransforms);
+		check((StartIndex + NumToCopy) <= NumTransforms);
 		check(StartIndex >= 0 && StartIndex <= NumTransforms);
 		check(EndIndex <= NumTransforms);
 
diff --git a/Engine/Plugins/Experimental/AnimNext/Source/AnimNext/Public/LODPose.h b/Engine/Plugins/Experimental/AnimNext/Source/AnimNext/Public/LODPose.h
--- a/Engine/Plugins/Experimental/AnimNext/Source/AnimNext/Public/LODPose.h
+++ b/Engine/Plugins/Experimental/AnimNext/Source/AnimNext/Public/LODPose.h
@@ -34,32 +34,53 @@ struct FLODPose
 
 	FLODPose() = default;
 
-	void CopyFrom(const FLODPose& Other)
+	void CopyFrom(const FLODPose& SourcePose)
 	{
-		check(LODLevel == Other.LODLevel);
-		check(RefPose && Other.RefPose);
-		check(RefPose->ReferenceLocalTransforms.Num() == Other.RefPose->ReferenceLocalTransforms.Num());
-		check(LocalTransformsView.Num() == Other.LocalTransformsView.Num());
+		check(RefPose && SourcePose.RefPose);
+		check(RefPose->ReferenceLocalTransforms.Num() == SourcePose.RefPose->ReferenceLocalTransforms.Num());
 
-		// Copy over the flags and transforms from our source
-		Flags = Other.Flags;
-		CopyTransforms(LocalTransformsView, Other.LocalTransformsView);
+		// Copy over the flags from our source.
+		Flags = SourcePose.Flags;
+
+		if (SourcePose.LODLevel == LODLevel)
+		{
+			// LOD levels match, just copy the full set of bone transforms.
+			CopyTransforms(LocalTransformsView, SourcePose.LocalTransformsView);
+		}
+		else if (SourcePose.LODLevel < LODLevel)
+		{
+			// The source pose is set to a higher-quality LOD level and contains more bone transforms than we need. Just copy the ones we actually need.
+			CopyTransforms(LocalTransformsView, SourcePose.LocalTransformsView, /*StartIndex=*/0, LocalTransformsView.Num());
+		}
+		else
+		{
+			// The source pose is missing transforms as it is set to a lower-quality LOD level, initialize the missing bone transforms with the reference pose transforms.
+			const int32 NumSourceBones = SourcePose.LocalTransformsView.Num();
+			const int32 NumTargetBones = LocalTransformsView.Num();
+			const int32 NumAdditionalBones = NumSourceBones - NumTargetBones;
+			SetRefPose(SourcePose.IsAdditive(), NumSourceBones, NumAdditionalBones);
+			CopyTransforms(LocalTransformsView, SourcePose.LocalTransformsView, /*StartIndex=*/0, NumSourceBones);
+		}
 	}
 
-	void SetRefPose(bool bAdditive = false)
+	void SetRefPose(bool bAdditive = false, int32 StartIndex = 0, int32 NumTransformsToCopy = -1)
 	{
 		const int32 NumTransforms = LocalTransformsView.Num();
-
 		if (NumTransforms > 0)
 		{
+			if (NumTransformsToCopy < 0)
+			{
+				NumTransformsToCopy = NumTransforms - StartIndex;
+			}
+
 			if (bAdditive)
 			{
-				SetIdentity(bAdditive);
+				SetIdentity(bAdditive, StartIndex, NumTransformsToCopy);
 			}
 			else
 			{
 				check(RefPose != nullptr);
-				CopyTransforms(LocalTransformsView, RefPose->ReferenceLocalTransforms.GetConstView(), 0, NumTransforms);
+				CopyTransforms(LocalTransformsView, RefPose->ReferenceLocalTransforms.GetConstView(), StartIndex, NumTransformsToCopy);
 			}
 		}
 
@@ -72,9 +93,9 @@ struct FLODPose
 		return *RefPose;
 	}
 
-	void SetIdentity(bool bAdditive = false)
+	void SetIdentity(bool bAdditive = false, int32 StartIndex = 0, int32 NumTransformsToSet = -1)
 	{
-		UE::AnimNext::SetIdentity(LocalTransformsView, bAdditive);
+		UE::AnimNext::SetIdentity(LocalTransformsView, bAdditive, StartIndex, NumTransformsToSet);
 	}
 
 	int32 GetNumBones() const
diff --git a/Engine/Plugins/Experimental/AnimNext/Source/AnimNext/Public/TransformArrayOperations.h b/Engine/Plugins/Experimental/AnimNext/Source/AnimNext/Public/TransformArrayOperations.h
--- a/Engine/Plugins/Experimental/AnimNext/Source/AnimNext/Public/TransformArrayOperations.h
+++ b/Engine/Plugins/Experimental/AnimNext/Source/AnimNext/Public/TransformArrayOperations.h
@@ -9,7 +9,9 @@ namespace UE::AnimNext
 {
 	// Sets the identity (regular or additive) over the whole destination array
 	ANIMNEXT_API void SetIdentity(const FTransformArrayAoSView& Dest, bool bIsAdditive);
+	ANIMNEXT_API void SetIdentity(const FTransformArrayAoSView& Dest, bool bIsAdditive, int32 StartIndex, int32 NumTransformsToSet);
 	ANIMNEXT_API void SetIdentity(const FTransformArraySoAView& Dest, bool bIsAdditive);
+	ANIMNEXT_API void SetIdentity(const FTransformArraySoAView& Dest, bool bIsAdditive, int32 StartIndex, int32 NumTransformsToSet);
 
 	// Copies the specified number of transforms from a source into a destination starting at the specified start index
 	// If NumToCopy is -1, we copy until the end
