Make C++ AnimNext API return const views / ptrs for thread safety.
Also add methods to modify those types via delegates that is thread safe.

#tests Nonunity, PIE run, will test PIE use after robomerge

diff --git a/Engine/Plugins/Experimental/AnimNext/Source/AnimNext/Internal/Component/AnimNextComponent.h b/Engine/Plugins/Experimental/AnimNext/Source/AnimNext/Internal/Component/AnimNextComponent.h
--- a/Engine/Plugins/Experimental/AnimNext/Source/AnimNext/Internal/Component/AnimNextComponent.h
+++ b/Engine/Plugins/Experimental/AnimNext/Source/AnimNext/Internal/Component/AnimNextComponent.h
@@ -60,6 +60,7 @@ public:
 	/**
 	 * Module variable value getters.
 	 * Numeric types (bool, (u)int32, (u)int64, float, double) support type conversion.
+	 * Struct & Object types will be const, if you need to modify them use SetVariable[Struct, Object]Ref
 	 * @param    Name     The name of the variable to get
 	 * @return            The variable value on success, None otherwise
 	 */
@@ -68,15 +69,15 @@ public:
 	ANIMNEXT_API TOptional<int32> GetVariableInt32(const FName Name) const;
 	ANIMNEXT_API TOptional<uint32> GetVariableUInt32(const FName Name) const;
 	ANIMNEXT_API TOptional<int64> GetVariableInt64(const FName Name) const;
-	ANIMNEXT_API TOptional<uint64> GetVariableUInt64(const FName Name) const;
+	ANIMNEXT_API TOptional<uint64> GetVariableUInt64(const FName Name) const; 
 	ANIMNEXT_API TOptional<float> GetVariableFloat(const FName Name) const;
 	ANIMNEXT_API TOptional<double> GetVariableDouble(const FName Name) const;
 	ANIMNEXT_API TOptional<FName> GetVariableName(const FName Name) const;
 	ANIMNEXT_API TOptional<FString> GetVariableString(const FName Name) const;
 	ANIMNEXT_API TOptional<uint8> GetVariableEnum(const FName Name, const UEnum* RequestedEnum) const;
-	ANIMNEXT_API TOptional<FStructView> GetVariableStruct(const FName Name, const UScriptStruct* RequestedStruct = nullptr) const;
-	ANIMNEXT_API TOptional<UObject*> GetVariableObject(const FName Name, const UClass* RequestedClass = nullptr) const;
-	ANIMNEXT_API TOptional<UClass*> GetVariableClass(const FName Name) const;
+	ANIMNEXT_API TOptional<FConstStructView> GetVariableStruct(const FName Name, const UScriptStruct* RequestedStruct = nullptr) const;
+	ANIMNEXT_API TOptional<const UObject*> GetVariableObject(const FName Name, const UClass* RequestedClass = nullptr) const;
+	ANIMNEXT_API TOptional<const UClass*> GetVariableClass(const FName Name) const;
 	ANIMNEXT_API TOptional<FSoftObjectPath> GetVariableSoftPath(const FName Name) const;
 
 	/** @return enum value of specified type. */
@@ -95,40 +96,40 @@ public:
 
 	/** @return struct reference of specified type. */
 	template <typename T>
-	TOptional<T*> GetVariableStruct(const FName Name) const
+	TOptional<const T*> GetVariableStruct(const FName Name) const
 	{
-		TOptional<FStructView> Result = GetVariableStruct(Name, TBaseStructure<T>::Get());
+		TOptional<FConstStructView> Result = GetVariableStruct(Name, TBaseStructure<T>::Get());
 		if (!Result.IsSet())
 		{
-			return TOptional<T*>();
+			return TOptional<const T*>();
 		}
-		if (T* ValuePtr = Result.GetValue().GetPtr<T>())
+		if (const T* ValuePtr = Result.GetValue().GetPtr<const T>())
 		{
-			return TOptional<T*>(ValuePtr);
+			return TOptional<const T*>(ValuePtr);
 		}
-		return TOptional<T*>();
+		return TOptional<const T*>();
 	}
 
 	/** @return object pointer value of specified type. */
 	template <typename T>
-	TOptional<T*> GetVariableObject(const FName Name) const
+	TOptional<const T*> GetVariableObject(const FName Name) const
 	{
 		static_assert(TIsDerivedFrom<T, UObject>::Value, "Should only call this with object types");
 
-		TOptional<UObject*> Result = GetVariableObject(Name, T::StaticClass());
+		TOptional<const UObject*> Result = GetVariableObject(Name, T::StaticClass());
 		if (!Result.IsSet())
 		{
-			return TOptional<T*>();
+			return TOptional<const T*>();
 		}
 		if (Result.GetValue() == nullptr)
 		{
-			return TOptional<T*>(nullptr);
+			return TOptional<const T*>(nullptr);
 		}
-		if (T* Object = Cast<T>(Result.GetValue()))
+		if (const T* Object = Cast<const T>(Result.GetValue()))
 		{
-			return TOptional<T*>(Object);
+			return TOptional<const T*>(Object);
 		}
-		return TOptional<T*>();
+		return TOptional<const T*>();
 	}
 
 	/**
@@ -149,9 +150,12 @@ public:
 	ANIMNEXT_API bool SetVariableString(const FName Name, const FString& InValue);
 	ANIMNEXT_API bool SetVariableEnum(const FName Name, const uint8 InValue, const UEnum* Enum);
 	ANIMNEXT_API bool SetVariableStruct(const FName Name, FConstStructView InValue);
+	ANIMNEXT_API bool SetVariableStructRef(const FName Name, TFunctionRef<void(FStructView)> InStructRefSetter, const UScriptStruct* RequestedStruct = nullptr);
 	ANIMNEXT_API bool SetVariableObject(const FName Name, UObject* InValue);
+	ANIMNEXT_API bool SetVariableObjectRef(const FName Name, TFunctionRef<void(UObject*)> InObjectRefSetter, const UClass* RequestedClass = nullptr);
 	ANIMNEXT_API bool SetVariableClass(const FName Name, UClass* InValue);
 	ANIMNEXT_API bool SetVariableSoftPath(const FName Name, const FSoftObjectPath& InValue);
+	ANIMNEXT_API bool SetVariableArrayRef(const FName Name, TFunctionRef<void(FPropertyBagArrayRef&)> InArrayRefSetter);
 
 	/** Sets enum variable with specified type. */
 	template <typename T>
diff --git a/Engine/Plugins/Experimental/AnimNext/Source/AnimNext/Private/Component/AnimNextComponent.cpp b/Engine/Plugins/Experimental/AnimNext/Source/AnimNext/Private/Component/AnimNextComponent.cpp
--- a/Engine/Plugins/Experimental/AnimNext/Source/AnimNext/Private/Component/AnimNextComponent.cpp
+++ b/Engine/Plugins/Experimental/AnimNext/Source/AnimNext/Private/Component/AnimNextComponent.cpp
@@ -369,12 +369,12 @@ TOptional<uint8> UAnimNextComponent::GetVariableEnum(const FName Name, const UEn
 	return Result.HasValue() ? TOptional<uint8>(Result.GetValue()) : TOptional<uint8>();
 }
 
-TOptional<FStructView> UAnimNextComponent::GetVariableStruct(const FName Name, const UScriptStruct* RequestedStruct) const
+TOptional<FConstStructView> UAnimNextComponent::GetVariableStruct(const FName Name, const UScriptStruct* RequestedStruct) const
 {
 	const int32* IndexPtr = PublicVariablesProxyMap.Find(Name);
 	if (IndexPtr == nullptr)
 	{
-		return TOptional<FStructView>();
+		return TOptional<FConstStructView>();
 	}
 
 	FRWScopeLock Lock(PublicVariablesLock, SLT_ReadOnly);
@@ -383,15 +383,15 @@ TOptional<FStructView> UAnimNextComponent::GetVariableStruct(const FName Name, c
 	const FPropertyBagPropertyDesc& PropertyDesc = ProxyDescs[*IndexPtr];
 
 	TValueOrError<FStructView, EPropertyBagResult> Result = PublicVariablesProxy.Data.GetValueStruct(PropertyDesc, RequestedStruct);
-	return Result.HasValue() ? TOptional<FStructView>(Result.GetValue()) : TOptional<FStructView>();
+	return Result.HasValue() ? TOptional<FConstStructView>(Result.GetValue()) : TOptional<FConstStructView>();
 }
 
-TOptional<UObject*> UAnimNextComponent::GetVariableObject(const FName Name, const UClass* RequestedClass) const
+TOptional<const UObject*> UAnimNextComponent::GetVariableObject(const FName Name, const UClass* RequestedClass) const
 {
 	const int32* IndexPtr = PublicVariablesProxyMap.Find(Name);
 	if (IndexPtr == nullptr)
 	{
-		return TOptional<UObject*>();
+		return TOptional<const UObject*>();
 	}
 
 	FRWScopeLock Lock(PublicVariablesLock, SLT_ReadOnly);
@@ -400,15 +400,15 @@ TOptional<UObject*> UAnimNextComponent::GetVariableObject(const FName Name, cons
 	const FPropertyBagPropertyDesc& PropertyDesc = ProxyDescs[*IndexPtr];
 
 	TValueOrError<UObject*, EPropertyBagResult> Result = PublicVariablesProxy.Data.GetValueObject(PropertyDesc, RequestedClass);
-	return Result.HasValue() ? TOptional<UObject*>(Result.GetValue()) : TOptional<UObject*>();
+	return Result.HasValue() ? TOptional<const UObject*>(Result.GetValue()) : TOptional<const UObject*>();
 }
 
-TOptional<UClass*> UAnimNextComponent::GetVariableClass(const FName Name) const
+TOptional<const UClass*> UAnimNextComponent::GetVariableClass(const FName Name) const
 {
 	const int32* IndexPtr = PublicVariablesProxyMap.Find(Name);
 	if (IndexPtr == nullptr)
 	{
-		return TOptional<UClass*>();
+		return TOptional<const UClass*>();
 	}
 
 	FRWScopeLock Lock(PublicVariablesLock, SLT_ReadOnly);
@@ -417,7 +417,7 @@ TOptional<UClass*> UAnimNextComponent::GetVariableClass(const FName Name) const
 	const FPropertyBagPropertyDesc& PropertyDesc = ProxyDescs[*IndexPtr];
 
 	TValueOrError<UClass*, EPropertyBagResult> Result = PublicVariablesProxy.Data.GetValueClass(PropertyDesc);
-	return Result.HasValue() ? TOptional<UClass*>(Result.GetValue()) : TOptional<UClass*>();
+	return Result.HasValue() ? TOptional<const UClass*>(Result.GetValue()) : TOptional<const UClass*>();
 }
 
 TOptional<FSoftObjectPath> UAnimNextComponent::GetVariableSoftPath(const FName Name) const
@@ -713,6 +713,31 @@ bool UAnimNextComponent::SetVariableStruct(const FName Name, FConstStructView In
 	return Result == EPropertyBagResult::Success;
 }
 
+bool UAnimNextComponent::SetVariableStructRef(const FName Name, TFunctionRef<void(FStructView)> InStructRefSetter, const UScriptStruct* RequestedStruct)
+{
+	const int32* IndexPtr = PublicVariablesProxyMap.Find(Name);
+	if (IndexPtr == nullptr)
+	{
+		return false;
+	}
+
+	FRWScopeLock Lock(PublicVariablesLock, SLT_Write);
+
+	TConstArrayView<FPropertyBagPropertyDesc> ProxyDescs = PublicVariablesProxy.Data.GetPropertyBagStruct()->GetPropertyDescs();
+	const FPropertyBagPropertyDesc& PropertyDesc = ProxyDescs[*IndexPtr];
+
+	TValueOrError<FStructView, EPropertyBagResult> Result = PublicVariablesProxy.Data.GetValueStruct(PropertyDesc, RequestedStruct);
+	if (Result.HasValue())
+	{
+		InStructRefSetter(Result.GetValue());
+		PublicVariablesProxy.DirtyFlags[*IndexPtr] = true;
+		PublicVariablesProxy.bIsDirty = true;
+		return true;
+	}
+
+	return false;
+}
+
 bool UAnimNextComponent::SetVariableObject(const FName Name, UObject* InValue)
 {
 	const int32* IndexPtr = PublicVariablesProxyMap.Find(Name);
@@ -736,6 +761,31 @@ bool UAnimNextComponent::SetVariableObject(const FName Name, UObject* InValue)
 	return Result == EPropertyBagResult::Success;
 }
 
+bool UAnimNextComponent::SetVariableObjectRef(const FName Name, TFunctionRef<void(UObject*)> InObjectRefSetter, const UClass* RequestedClass)
+{
+	const int32* IndexPtr = PublicVariablesProxyMap.Find(Name);
+	if (IndexPtr == nullptr)
+	{
+		return false;
+	}
+
+	FRWScopeLock Lock(PublicVariablesLock, SLT_Write);
+
+	TConstArrayView<FPropertyBagPropertyDesc> ProxyDescs = PublicVariablesProxy.Data.GetPropertyBagStruct()->GetPropertyDescs();
+	const FPropertyBagPropertyDesc& PropertyDesc = ProxyDescs[*IndexPtr];
+
+	TValueOrError<UObject*, EPropertyBagResult> Result = PublicVariablesProxy.Data.GetValueObject(PropertyDesc, RequestedClass);
+	if (Result.HasValue())
+	{
+		InObjectRefSetter(Result.GetValue());
+		PublicVariablesProxy.DirtyFlags[*IndexPtr] = true;
+		PublicVariablesProxy.bIsDirty = true;
+		return true;
+	}
+
+	return false;
+}
+
 bool UAnimNextComponent::SetVariableClass(const FName Name, UClass* InValue)
 {
 	const int32* IndexPtr = PublicVariablesProxyMap.Find(Name);
@@ -782,6 +832,31 @@ bool UAnimNextComponent::SetVariableSoftPath(const FName Name, const FSoftObject
 	return Result == EPropertyBagResult::Success;
 }
 
+bool UAnimNextComponent::SetVariableArrayRef(const FName Name, TFunctionRef<void(FPropertyBagArrayRef&)> InArrayRefSetter)
+{
+	const int32* IndexPtr = PublicVariablesProxyMap.Find(Name);
+	if (IndexPtr == nullptr)
+	{
+		return false;
+	}
+
+	FRWScopeLock Lock(PublicVariablesLock, SLT_Write);
+
+	TConstArrayView<FPropertyBagPropertyDesc> ProxyDescs = PublicVariablesProxy.Data.GetPropertyBagStruct()->GetPropertyDescs();
+	const FPropertyBagPropertyDesc& PropertyDesc = ProxyDescs[*IndexPtr];
+
+	TValueOrError<FPropertyBagArrayRef, EPropertyBagResult> Result = PublicVariablesProxy.Data.GetMutableArrayRef(PropertyDesc);
+	if (Result.HasValue())
+	{
+		InArrayRefSetter(Result.GetValue());
+		PublicVariablesProxy.DirtyFlags[*IndexPtr] = true;
+		PublicVariablesProxy.bIsDirty = true;
+		return true;
+	}
+
+	return false;
+}
+
 void UAnimNextComponent::SetEnabled(bool bEnabled)
 {
 	if(Subsystem)
diff --git a/Engine/Source/Runtime/CoreUObject/Private/StructUtils/PropertyBag.cpp b/Engine/Source/Runtime/CoreUObject/Private/StructUtils/PropertyBag.cpp
--- a/Engine/Source/Runtime/CoreUObject/Private/StructUtils/PropertyBag.cpp
+++ b/Engine/Source/Runtime/CoreUObject/Private/StructUtils/PropertyBag.cpp
@@ -2295,6 +2295,94 @@ TValueOrError<FSoftObjectPath, EPropertyBagResult> FInstancedPropertyBag::GetVal
 	return MakeValue(std::move(ReturnValue));
 }
 
+TValueOrError<FPropertyBagArrayRef, EPropertyBagResult> FInstancedPropertyBag::GetMutableArrayRef(const FPropertyBagPropertyDesc& Desc)
+{
+	if (!OwnsPropertyDesc(Desc))
+	{
+		return MakeError(EPropertyBagResult::PropertyNotFound);
+	}
+	check(Desc.CachedProperty);
+
+	if (Desc.ContainerTypes.GetFirstContainerType() != EPropertyBagContainerType::Array)
+	{
+		return MakeError(EPropertyBagResult::TypeMismatch);
+	}
+
+	const void* Address = GetValueAddress(&Desc);
+	if (Address == nullptr)
+	{
+		return MakeError(EPropertyBagResult::PropertyNotFound);
+	}
+
+	return MakeValue(FPropertyBagArrayRef(Desc, Address));
+}
+
+TValueOrError<const FPropertyBagArrayRef, EPropertyBagResult> FInstancedPropertyBag::GetArrayRef(const FPropertyBagPropertyDesc& Desc) const
+{
+	if (!OwnsPropertyDesc(Desc))
+	{
+		return MakeError(EPropertyBagResult::PropertyNotFound);
+	}
+	check(Desc.CachedProperty);
+
+	if (Desc.ContainerTypes.GetFirstContainerType() != EPropertyBagContainerType::Array)
+	{
+		return MakeError(EPropertyBagResult::TypeMismatch);
+	}
+
+	const void* Address = GetValueAddress(&Desc);
+	if (Address == nullptr)
+	{
+		return MakeError(EPropertyBagResult::PropertyNotFound);
+	}
+
+	return MakeValue(FPropertyBagArrayRef(Desc, Address));
+}
+
+TValueOrError<FPropertyBagSetRef, EPropertyBagResult> FInstancedPropertyBag::GetMutableSetRef(const FPropertyBagPropertyDesc& Desc)
+{
+	if (!OwnsPropertyDesc(Desc))
+	{
+		return MakeError(EPropertyBagResult::PropertyNotFound);
+	}
+	check(Desc.CachedProperty);
+
+	if (Desc.ContainerTypes.GetFirstContainerType() != EPropertyBagContainerType::Set)
+	{
+		return MakeError(EPropertyBagResult::TypeMismatch);
+	}
+
+	const void* Address = GetValueAddress(&Desc);
+	if (Address == nullptr)
+	{
+		return MakeError(EPropertyBagResult::PropertyNotFound);
+	}
+
+	return MakeValue(FPropertyBagSetRef(Desc, Address));
+}
+
+TValueOrError<const FPropertyBagSetRef, EPropertyBagResult> FInstancedPropertyBag::GetSetRef(const FPropertyBagPropertyDesc& Desc) const
+{
+	if (!OwnsPropertyDesc(Desc))
+	{
+		return MakeError(EPropertyBagResult::PropertyNotFound);
+	}
+	check(Desc.CachedProperty);
+
+	if (Desc.ContainerTypes.GetFirstContainerType() != EPropertyBagContainerType::Set)
+	{
+		return MakeError(EPropertyBagResult::TypeMismatch);
+	}
+
+	const void* Address = GetValueAddress(&Desc);
+	if (Address == nullptr)
+	{
+		return MakeError(EPropertyBagResult::PropertyNotFound);
+	}
+
+	return MakeValue(FPropertyBagSetRef(Desc, Address));
+}
+
 EPropertyBagResult FInstancedPropertyBag::SetValueBool(const FName Name, const bool bInValue)
 {
 	const FPropertyBagPropertyDesc* Desc = FindPropertyDescByName(Name);
diff --git a/Engine/Source/Runtime/CoreUObject/Public/StructUtils/PropertyBag.h b/Engine/Source/Runtime/CoreUObject/Public/StructUtils/PropertyBag.h
--- a/Engine/Source/Runtime/CoreUObject/Public/StructUtils/PropertyBag.h
+++ b/Engine/Source/Runtime/CoreUObject/Public/StructUtils/PropertyBag.h
@@ -573,6 +573,10 @@ struct COREUOBJECT_API FInstancedPropertyBag
 	TValueOrError<UObject*, EPropertyBagResult> GetValueObject(const FPropertyBagPropertyDesc& Desc, const UClass* RequestedClass = nullptr) const;
 	TValueOrError<UClass*, EPropertyBagResult> GetValueClass(const FPropertyBagPropertyDesc& Desc) const;
 	TValueOrError<FSoftObjectPath, EPropertyBagResult> GetValueSoftPath(const FPropertyBagPropertyDesc& Desc) const;
+	TValueOrError<FPropertyBagArrayRef, EPropertyBagResult> GetMutableArrayRef(const FPropertyBagPropertyDesc& Desc);
+	TValueOrError<const FPropertyBagArrayRef, EPropertyBagResult> GetArrayRef(const FPropertyBagPropertyDesc& Desc) const;
+	TValueOrError<FPropertyBagSetRef, EPropertyBagResult> GetMutableSetRef(const FPropertyBagPropertyDesc& Desc);
+	TValueOrError<const FPropertyBagSetRef, EPropertyBagResult> GetSetRef(const FPropertyBagPropertyDesc& Desc) const;
 
 	/** @return enum value of specified type. */
 	template <typename T>
