Fix invalid memory access on app exit in AnimNext data registry

diff --git a/Engine/Plugins/Experimental/AnimNext/Source/AnimNext/Private/DataRegistry.cpp b/Engine/Plugins/Experimental/AnimNext/Source/AnimNext/Private/DataRegistry.cpp
--- a/Engine/Plugins/Experimental/AnimNext/Source/AnimNext/Private/DataRegistry.cpp
+++ b/Engine/Plugins/Experimental/AnimNext/Source/AnimNext/Private/DataRegistry.cpp
@@ -39,6 +39,14 @@ namespace UE::AnimNext
 	{
 		FCoreUObjectDelegates::GetPostGarbageCollect().Remove(PostGarbageCollectHandle);
 
+		// When the module unloads, we cannot use ScriptStructs and our ParamTypes to free the memory
+		// we allocated. To avoid crashing on an invalid access, we clear the type map which means that
+		// destructors won't run for allocated data. If they were to free memory, it will leak on app exit.
+		{
+			FRWScopeLock ScopedDataTypeDefsLock(GAnimationDataRegistry->DataTypeDefsLock, SLT_Write);
+			GAnimationDataRegistry->DataTypeDefs.Empty();
+		}
+
 		GAnimationDataRegistry->ReleaseReferencePoseData(); // release any registered poses
 
 		check(GAnimationDataRegistry->AllocatedBlocks.Num() == 0); // any other data should have been released at this point
@@ -189,8 +197,11 @@ void FDataRegistry::FreeAllocatedBlock(Private::FAllocatedBlock* AllocatedBlock)
 
 			{
 				FRWScopeLock ScopedDataTypeDefsLock(DataTypeDefsLock, SLT_ReadOnly);
-				FDataTypeDef* TypeDef = DataTypeDefs.Find(AllocatedBlock->Type);
-				if (ensure(TypeDef != nullptr))
+				// We are empty on app exit and thus cannot run destructors (param type is invalid)
+				// When this occurs, we leak memory on purpose
+				const FDataTypeDef* TypeDef = !DataTypeDefs.IsEmpty() ? DataTypeDefs.Find(AllocatedBlock->Type) : nullptr;
+				if (!DataTypeDefs.IsEmpty()
+					&& ensure(TypeDef != nullptr))
 				{
 					TypeDef->DestroyTypeFn((uint8*)AllocatedBlock->Memory, AllocatedBlock->NumElem);
 				}
