AnimNext RigLogic trait crash

RigLogic crashed when the mesh's DNA for a given LOD had a bone enabled that wasn't part of the given pose's LOD. Filter these in the init-phase already in order to prevent out of bound access and avoid runtime checks.

diff --git a/Engine/Plugins/Experimental/AnimNext/Source/AnimNext/Private/DecompressionTools.cpp b/Engine/Plugins/Experimental/AnimNext/Source/AnimNext/Private/DecompressionTools.cpp
--- a/Engine/Plugins/Experimental/AnimNext/Source/AnimNext/Private/DecompressionTools.cpp
+++ b/Engine/Plugins/Experimental/AnimNext/Source/AnimNext/Private/DecompressionTools.cpp
@@ -314,7 +314,7 @@ void FDecompressionTools::DecompressPose(FLODPose& OutAnimationPoseData,
 {
 	const FReferencePose& ReferencePose = OutAnimationPoseData.GetRefPose();
 	const TArrayView<const FBoneIndexType> LODBoneIndexToSkeletonBoneIndexMap = OutAnimationPoseData.GetLODBoneIndexToMeshBoneIndexMap();
-	const TArrayView<const FBoneIndexType> SkeletonToLODBoneIndexes = ReferencePose.GetSkeletonBoneIndexToLODBoneIndexMap(0); // Full list of Skeleton to LOD conversion
+	const TArrayView<const FBoneIndexType> SkeletonToLODBoneIndexes = ReferencePose.GetSkeletonBoneIndexToLODBoneIndexMap(); // Full list of Skeleton to LOD conversion
 	const int32 NumLODBoneIndexes = LODBoneIndexToSkeletonBoneIndexMap.Num();
 
 	const int32 NumTracks = CompressedData.CompressedTrackToSkeletonMapTable.Num();
diff --git a/Engine/Plugins/Experimental/AnimNext/Source/AnimNext/Private/Graph/RemapPoseData.cpp b/Engine/Plugins/Experimental/AnimNext/Source/AnimNext/Private/Graph/RemapPoseData.cpp
--- a/Engine/Plugins/Experimental/AnimNext/Source/AnimNext/Private/Graph/RemapPoseData.cpp
+++ b/Engine/Plugins/Experimental/AnimNext/Source/AnimNext/Private/Graph/RemapPoseData.cpp
@@ -48,9 +48,10 @@ void FRemapPoseData::Reinit(const UE::AnimNext::FReferencePose& InSourceRefPose,
 	SourceToTargetBoneIndexMapPerLOD.Empty();
 	SourceToTargetBoneIndexMapPerLOD.AddDefaulted(TargetNumLODLevels);
 
+	const TArrayView<const FBoneIndexType> SourceSkeletonIndexToPoseIndex = InSourceRefPose.GetSkeletonBoneIndexToLODBoneIndexMap(); // skeleton -> pose
+
 	for (int32 LODLevel = 0; LODLevel < TargetNumLODLevels; ++LODLevel)
 	{
-		const TArrayView<const FBoneIndexType> SourceSkeletonIndexToPoseIndex = InSourceRefPose.GetSkeletonBoneIndexToLODBoneIndexMap(LODLevel); // skeleton -> pose
 		const TArrayView<const FBoneIndexType> TargetPoseToMeshBoneIndexMap = InTargetRefPose.GetLODBoneIndexToMeshBoneIndexMap(LODLevel); // pose -> mesh
 		const int32 SourceNumBonesInLOD = SourceRefPose->GetNumBonesForLOD(LODLevel);
 
diff --git a/Engine/Plugins/Experimental/AnimNext/Source/AnimNext/Private/Graph/RigUnit_AnimNextWriteSkeletalMeshComponentPose.cpp b/Engine/Plugins/Experimental/AnimNext/Source/AnimNext/Private/Graph/RigUnit_AnimNextWriteSkeletalMeshComponentPose.cpp
--- a/Engine/Plugins/Experimental/AnimNext/Source/AnimNext/Private/Graph/RigUnit_AnimNextWriteSkeletalMeshComponentPose.cpp
+++ b/Engine/Plugins/Experimental/AnimNext/Source/AnimNext/Private/Graph/RigUnit_AnimNextWriteSkeletalMeshComponentPose.cpp
@@ -47,6 +47,7 @@ FRigUnit_AnimNextWriteSkeletalMeshComponentPose_Execute()
 
 	FMemMark MemMark(FMemStack::Get());
 
+	// TODO: The LOD level we got here is not guaranteed to be stable. We should find another way to get/bind the skeletal mesh component's LOD level.
 	const int32 SkeletalMeshComponentLODLevel = SkeletalMeshComponent->GetPredictedLODLevel();
 	const int32 PoseLODLevel = Pose.LODPose.LODLevel;
 
diff --git a/Engine/Plugins/Experimental/AnimNext/Source/AnimNext/Public/LODPose.h b/Engine/Plugins/Experimental/AnimNext/Source/AnimNext/Public/LODPose.h
--- a/Engine/Plugins/Experimental/AnimNext/Source/AnimNext/Public/LODPose.h
+++ b/Engine/Plugins/Experimental/AnimNext/Source/AnimNext/Public/LODPose.h
@@ -141,9 +141,9 @@ struct FLODPose
 
 	const TArrayView<const FBoneIndexType> GetSkeletonBoneIndexToLODBoneIndexMap() const
 	{
-		if (LODLevel != INVALID_LOD_LEVEL && RefPose != nullptr)
+		if (RefPose != nullptr)
 		{
-			return RefPose->GetSkeletonBoneIndexToLODBoneIndexMap(LODLevel);
+			return RefPose->GetSkeletonBoneIndexToLODBoneIndexMap();
 		}
 		else
 		{
diff --git a/Engine/Plugins/Experimental/AnimNext/Source/AnimNext/Public/ReferencePose.h b/Engine/Plugins/Experimental/AnimNext/Source/AnimNext/Public/ReferencePose.h
--- a/Engine/Plugins/Experimental/AnimNext/Source/AnimNext/Public/ReferencePose.h
+++ b/Engine/Plugins/Experimental/AnimNext/Source/AnimNext/Public/ReferencePose.h
@@ -49,15 +49,17 @@ struct TReferencePose
 	// When fast path is enabled, we have a single LOD entry that we truncate to the number of bones for each LOD
 	TArray<TArray<FBoneIndexType, AllocatorType>, AllocatorType> LODBoneIndexToSkeletonBoneIndexMapPerLOD;
 
-	// List of skeleton bone indices for each LOD
+	// List of skeleton bone indices
 	// Each list of skeleton bone indices is a mapping of: SkeletonBoneIndex -> LODSortedBoneIndex
-	// When fast path is enabled, we have a single LOD entry
-	TArray<TArray<FBoneIndexType, AllocatorType>, AllocatorType> SkeletonBoneIndexToLODBoneIndexMapPerLOD;
+	// Size of the map equals the number of mesh bones. Mesh bones that are not part of the pose hold a INDEX_NONE.
+	// Note: This map is LOD independent. Please use IsBoneEnabled() to check if a given bone is part of a LOD level.
+	TArray<FBoneIndexType, AllocatorType> SkeletonBoneIndexToLODBoneIndexMap;
 
-	// List of pose bone indices by mesh bone index for each LOD
-	// Each list of bone indices is a mapping of: SkeletalMeshBoneIndex -> LODSortedBoneIndex
-	// When fast path is enabled, we have a single LOD entry
-	TArray<TArray<FBoneIndexType, AllocatorType>, AllocatorType> MeshBoneIndexToLODBoneIndexMapPerLOD;
+	// List of pose bone indices by mesh bone index
+	// List of bone indices is a mapping of: SkeletalMeshBoneIndex -> LODSortedBoneIndex
+	// Size of the map equals the number of mesh bones. Mesh bones that are not part of the pose hold a INDEX_NONE.
+	// Note: This map is LOD independent. Please use IsBoneEnabled() to check if a given bone is part of a LOD level.
+	TArray<FBoneIndexType, AllocatorType> MeshBoneIndexToLODBoneIndexMap;
 
 	// Number of bones for each LOD
 	TArray<int32, AllocatorType> LODNumBones;
@@ -83,6 +85,11 @@ struct TReferencePose
 		return (LODLevel < NumLODS) ? LODNumBones[LODLevel] : (NumLODS > 0) ? LODNumBones[0] : 0;
 	}
 
+	bool IsBoneEnabled(int32 BoneIndex, int32 LODLevel) const
+	{
+		return (BoneIndex < GetNumBonesForLOD(LODLevel));
+	}
+
 	bool IsFastPath() const
 	{
 		return (GenerationFlags & EReferencePoseGenerationFlags::FastPath) != EReferencePoseGenerationFlags::None;
@@ -92,21 +99,21 @@ struct TReferencePose
 		, const TArray<TArray<FBoneIndexType>>& InLODBoneIndexToParentLODBoneIndexMapPerLOD
 		, const TArray<TArray<FBoneIndexType>>& InLODBoneIndexToMeshBoneIndexMapPerLOD
 		, const TArray<TArray<FBoneIndexType>>& InLODBoneIndexToSkeletonBoneIndexMapPerLOD
-		, const TArray<TArray<FBoneIndexType>>& InSkeletonBoneIndexToLODBoneIndexMapPerLOD
-		, const TArray<TArray<FBoneIndexType>>& InMeshBoneIndexToLODBoneIndexMapPerLOD
+		, const TArray<FBoneIndexType>& InSkeletonBoneIndexToLODBoneIndexMap
+		, const TArray<FBoneIndexType>& InMeshBoneIndexToLODBoneIndexMap
 		, const TArray<int32, AllocatorType>& InLODNumBones
 		, bool bFastPath = false)
 	{
 		const int32 NumBonesLOD0 = !InLODNumBones.IsEmpty() ? InLODNumBones[0] : 0;
 		const int32 NumBonesMesh = RefSkeleton.GetRefBoneInfo().Num();
-		
+
 		MeshBoneIndexToParentMeshBoneIndexMap.SetNum(NumBonesMesh);
 		ReferenceLocalTransforms.SetNum(NumBonesLOD0);
 		LODBoneIndexToParentLODBoneIndexMapPerLOD = InLODBoneIndexToParentLODBoneIndexMapPerLOD;
 		LODBoneIndexToMeshBoneIndexMapPerLOD = InLODBoneIndexToMeshBoneIndexMapPerLOD;
 		LODBoneIndexToSkeletonBoneIndexMapPerLOD = InLODBoneIndexToSkeletonBoneIndexMapPerLOD;
-		SkeletonBoneIndexToLODBoneIndexMapPerLOD = InSkeletonBoneIndexToLODBoneIndexMapPerLOD;
-		MeshBoneIndexToLODBoneIndexMapPerLOD = InMeshBoneIndexToLODBoneIndexMapPerLOD;
+		SkeletonBoneIndexToLODBoneIndexMap = InSkeletonBoneIndexToLODBoneIndexMap;
+		MeshBoneIndexToLODBoneIndexMap = InMeshBoneIndexToLODBoneIndexMap;
 		LODNumBones = InLODNumBones;
 
 		const TArray<FTransform>& RefBonePose = RefSkeleton.GetRefBonePose();
@@ -176,31 +183,15 @@ struct TReferencePose
 	}
 
 	// Returns a list of LOD bone indices, a mapping of: SkeletonBoneIndex -> LODSortedBoneIndex
-	const TArrayView<const FBoneIndexType> GetSkeletonBoneIndexToLODBoneIndexMap(int32 LODLevel) const
+	const TArrayView<const FBoneIndexType> GetSkeletonBoneIndexToLODBoneIndexMap() const
 	{
-		TArrayView<const FBoneIndexType> ArrayView;
-
-		if (LODLevel >= 0 && (IsFastPath() || LODLevel < SkeletonBoneIndexToLODBoneIndexMapPerLOD.Num()))
-		{
-			const int32 LODIndex = IsFastPath() ? 0 : LODLevel;
-			ArrayView = SkeletonBoneIndexToLODBoneIndexMapPerLOD[LODIndex];
-		}
-
-		return ArrayView;
+		return SkeletonBoneIndexToLODBoneIndexMap;
 	}
 
 	// Return a list of LOD bone indices, a mapping of: SkeletalMeshBoneIndex -> LODSortedBoneIndex
-	const TArrayView<const FBoneIndexType> GetMeshBoneIndexToLODBoneIndexMap(int32 LODLevel) const
+	const TArrayView<const FBoneIndexType> GetMeshBoneIndexToLODBoneIndexMap() const
 	{
-		TArrayView<const FBoneIndexType> ArrayView;
-
-		if (LODLevel >= 0 && (IsFastPath() || LODLevel < MeshBoneIndexToLODBoneIndexMapPerLOD.Num()))
-		{
-			const int32 LODIndex = IsFastPath() ? 0 : LODLevel;
-			ArrayView = MeshBoneIndexToLODBoneIndexMapPerLOD[LODIndex];
-		}
-
-		return ArrayView;
+		return MeshBoneIndexToLODBoneIndexMap;
 	}
 
 	int32 GetMeshBoneIndexFromLODBoneIndex(int32 LODBoneIndex) const
@@ -226,16 +217,14 @@ struct TReferencePose
 	// Translate a skeletal mesh bone index to a pose bone index.
 	int32 GetLODBoneIndexFromMeshBoneIndex(int32 MeshBoneIndex) const
 	{
-		const int32 NumBonesLOD0 = MeshBoneIndexToLODBoneIndexMapPerLOD[0].Num();
-		check(MeshBoneIndex < NumBonesLOD0);
-		return MeshBoneIndexToLODBoneIndexMapPerLOD[0][MeshBoneIndex];
+		check(MeshBoneIndex < MeshBoneIndexToLODBoneIndexMap.Num());
+		return MeshBoneIndexToLODBoneIndexMap[MeshBoneIndex];
 	}
 
 	int32 GetLODBoneIndexFromSkeletonBoneIndex(int32 SkeletionBoneIndex) const
 	{
-		const int32 NumBonesLOD0 = SkeletonBoneIndexToLODBoneIndexMapPerLOD[0].Num();
-		check(SkeletionBoneIndex < NumBonesLOD0);
-		return SkeletonBoneIndexToLODBoneIndexMapPerLOD[0][SkeletionBoneIndex];
+		check(SkeletionBoneIndex < SkeletonBoneIndexToLODBoneIndexMap.Num());
+		return SkeletonBoneIndexToLODBoneIndexMap[SkeletionBoneIndex];
 	}
 
 	FTransform GetRefPoseTransform(int32 LODBoneIndex) const
diff --git a/Engine/Plugins/Experimental/RigLogicAnimNext/Source/RigLogicAnimNext/Private/RigLogicInstanceData.cpp b/Engine/Plugins/Experimental/RigLogicAnimNext/Source/RigLogicAnimNext/Private/RigLogicInstanceData.cpp
--- a/Engine/Plugins/Experimental/RigLogicAnimNext/Source/RigLogicAnimNext/Private/RigLogicInstanceData.cpp
+++ b/Engine/Plugins/Experimental/RigLogicAnimNext/Source/RigLogicAnimNext/Private/RigLogicInstanceData.cpp
@@ -65,13 +65,13 @@ namespace UE::AnimNext
 		RigLogicToSkeletonBoneIndexMappingPerLOD.Empty();
 		RigLogicToSkeletonBoneIndexMappingPerLOD.SetNum(NumLODs);
 
+		const TArrayView<const FBoneIndexType> MeshToPoseBoneIndexMap = ReferencePose->GetMeshBoneIndexToLODBoneIndexMap(); // mesh -> pose
+
 		for (uint32 LODLevel = 0; LODLevel < NumLODs; ++LODLevel)
 		{
 			const TArray<uint16>& VariableJointIndices = CachedRigRuntimeContext->VariableJointIndicesPerLOD[LODLevel].Values;
 			RigLogicToSkeletonBoneIndexMappingPerLOD[LODLevel].Reserve(VariableJointIndices.Num());
 
-			const TArrayView<const FBoneIndexType> MeshToPoseBoneIndexMap = ReferencePose->GetMeshBoneIndexToLODBoneIndexMap(LODLevel); // mesh -> pose
-
 			for (const uint16 RigLogicJointIndex : VariableJointIndices)
 			{
 				// Get the mesh bone index from the RigLogic joint index. Index originating from a FindBoneName() on the mesh's reference skeleton.
@@ -80,8 +80,10 @@ namespace UE::AnimNext
 				{
 					// Convert the skeleton bone index to a pose bone index for the given LOD level.
 					const int32 PoseBoneIndex = MeshToPoseBoneIndexMap[MeshBoneIndex.GetInt()];
-
-					RigLogicToSkeletonBoneIndexMappingPerLOD[LODLevel].Add({ RigLogicJointIndex, PoseBoneIndex });
+					if (ReferencePose->IsBoneEnabled(PoseBoneIndex, LODLevel))
+					{
+						RigLogicToSkeletonBoneIndexMappingPerLOD[LODLevel].Add({ RigLogicJointIndex, PoseBoneIndex });
+					}
 				}
 				else
 				{
@@ -102,10 +104,10 @@ namespace UE::AnimNext
 		SparseDriverJointsToControlAttributesMapPerLOD.SetNum(NumLODs);
 		DenseDriverJointsToControlAttributesMapPerLOD.SetNum(NumLODs);
 
+		const TArrayView<const FBoneIndexType> MeshToPoseBoneIndexMap = ReferencePose->GetMeshBoneIndexToLODBoneIndexMap(); // mesh -> pose
+		
 		for (uint32 LODLevel = 0; LODLevel < NumLODs; ++LODLevel)
 		{
-			const TArrayView<const FBoneIndexType> MeshToPoseBoneIndexMap = ReferencePose->GetMeshBoneIndexToLODBoneIndexMap(LODLevel); // mesh -> pose
-
 			// Sparse mapping will likely remain empty so no reservation happens
 			for (const auto& Mapping : CachedDNAIndexMapping->DriverJointsToControlAttributesMap)
 			{
@@ -114,17 +116,19 @@ namespace UE::AnimNext
 				{
 					// Convert the mesh bone index to a pose bone index for the given LOD level.
 					const int32 PoseBoneIndex = MeshToPoseBoneIndexMap[MeshBoneIndex.GetInt()];
-
-					if ((Mapping.RotationX != INDEX_NONE) && (Mapping.RotationY != INDEX_NONE) && (Mapping.RotationZ != INDEX_NONE) && (Mapping.RotationW != INDEX_NONE))
-					{
-						DenseDriverJointsToControlAttributesMapPerLOD[LODLevel].Add({ PoseBoneIndex, Mapping.DNAJointIndex, Mapping.RotationX, Mapping.RotationY, Mapping.RotationZ, Mapping.RotationW });
-					}
-					else
+					if (ReferencePose->IsBoneEnabled(PoseBoneIndex, LODLevel))
 					{
-						SparseDriverJointsToControlAttributesMapPerLOD[LODLevel].Add({ PoseBoneIndex, Mapping.DNAJointIndex, Mapping.RotationX, Mapping.RotationY, Mapping.RotationZ, Mapping.RotationW });
+						if ((Mapping.RotationX != INDEX_NONE) && (Mapping.RotationY != INDEX_NONE) && (Mapping.RotationZ != INDEX_NONE) && (Mapping.RotationW != INDEX_NONE))
+						{
+							DenseDriverJointsToControlAttributesMapPerLOD[LODLevel].Add({ PoseBoneIndex, Mapping.DNAJointIndex, Mapping.RotationX, Mapping.RotationY, Mapping.RotationZ, Mapping.RotationW });
+						}
+						else
+						{
+							SparseDriverJointsToControlAttributesMapPerLOD[LODLevel].Add({ PoseBoneIndex, Mapping.DNAJointIndex, Mapping.RotationX, Mapping.RotationY, Mapping.RotationZ, Mapping.RotationW });
+						}
 					}
 				}
 			}
 		}
 	}
-} // namespace UE::AnimNext
\ No newline at end of file
+} // namespace UE::AnimNext
