Fix crash when duplicating AnimNext assets in the content browser

diff --git a/Engine/Plugins/Experimental/AnimNext/Source/AnimNext/Internal/AnimNextRigVMAsset.h b/Engine/Plugins/Experimental/AnimNext/Source/AnimNext/Internal/AnimNextRigVMAsset.h
--- a/Engine/Plugins/Experimental/AnimNext/Source/AnimNext/Internal/AnimNextRigVMAsset.h
+++ b/Engine/Plugins/Experimental/AnimNext/Source/AnimNext/Internal/AnimNextRigVMAsset.h
@@ -1,4 +1,4 @@
-﻿// Copyright Epic Games, Inc. All Rights Reserved.
+// Copyright Epic Games, Inc. All Rights Reserved.
 
 #pragma once
 
@@ -41,7 +41,8 @@ protected:
 	// UObject interface
 	virtual void BeginDestroy() override;
 	virtual void PostLoad() override;
-	virtual void GetAssetRegistryTags(FAssetRegistryTagsContext Context) const override; 
+	virtual void GetAssetRegistryTags(FAssetRegistryTagsContext Context) const override;
+	virtual void PreDuplicate(FObjectDuplicationParameters& DupParams) override;
 
 public:
 	// Get public variable defaults property bag
diff --git a/Engine/Plugins/Experimental/AnimNext/Source/AnimNext/Private/AnimNextRigVMAsset.cpp b/Engine/Plugins/Experimental/AnimNext/Source/AnimNext/Private/AnimNextRigVMAsset.cpp
--- a/Engine/Plugins/Experimental/AnimNext/Source/AnimNext/Private/AnimNextRigVMAsset.cpp
+++ b/Engine/Plugins/Experimental/AnimNext/Source/AnimNext/Private/AnimNextRigVMAsset.cpp
@@ -1,4 +1,4 @@
-﻿// Copyright Epic Games, Inc. All Rights Reserved.
+// Copyright Epic Games, Inc. All Rights Reserved.
 
 #include "AnimNextRigVMAsset.h"
 #include "RigVMRuntimeDataRegistry.h"
@@ -65,6 +65,18 @@ void UAnimNextRigVMAsset::GetAssetRegistryTags(FAssetRegistryTagsContext Context
 #endif // WITH_EDITOR
 }
 
+void UAnimNextRigVMAsset::PreDuplicate(FObjectDuplicationParameters& DupParams)
+{
+	Super::PreDuplicate(DupParams);
+
+#if WITH_EDITORONLY_DATA
+	if (EditorData)
+	{
+		EditorData->PreDuplicate(DupParams);
+	}
+#endif
+}
+
 TArray<FRigVMExternalVariable> UAnimNextRigVMAsset::GetExternalVariablesImpl(bool bFallbackToBlueprint) const
 {
 	TArray<FRigVMExternalVariable> ExternalVariables;
diff --git a/Engine/Plugins/Experimental/AnimNext/Source/AnimNextEditor/Private/AnimNextEditorModule.cpp b/Engine/Plugins/Experimental/AnimNext/Source/AnimNextEditor/Private/AnimNextEditorModule.cpp
--- a/Engine/Plugins/Experimental/AnimNext/Source/AnimNextEditor/Private/AnimNextEditorModule.cpp
+++ b/Engine/Plugins/Experimental/AnimNext/Source/AnimNextEditor/Private/AnimNextEditorModule.cpp
@@ -805,6 +805,13 @@ void FAnimNextEditorModule::RegisterWorkspaceDocumentTypes(Workspace::IWorkspace
 			UAnimNextRigVMAssetEditorData* EditorData = EdGraph->GetTypedOuter<UAnimNextRigVMAssetEditorData>();
 			if(EditorData)
 			{
+				// If we are dirty, make sure to reconstruct our nodes and recompile as needed
+				if (EditorData->bVMRecompilationRequired)
+				{
+					EditorData->ReconstructAllNodes();
+					EditorData->RecompileVMIfRequired();
+				}
+
 				EditorData->InteractionBracketFinished.RemoveAll(&InContext.WorkspaceEditor.Get());
 				EditorData->InteractionBracketFinished.AddSPLambda(&InContext.WorkspaceEditor.Get(), [WeakWorkspaceEditor](UAnimNextRigVMAssetEditorData* InEditorData)
 				{
diff --git a/Engine/Plugins/Experimental/AnimNext/Source/AnimNextUncookedOnly/Internal/AnimNextRigVMAssetEditorData.h b/Engine/Plugins/Experimental/AnimNext/Source/AnimNextUncookedOnly/Internal/AnimNextRigVMAssetEditorData.h
--- a/Engine/Plugins/Experimental/AnimNext/Source/AnimNextUncookedOnly/Internal/AnimNextRigVMAssetEditorData.h
+++ b/Engine/Plugins/Experimental/AnimNext/Source/AnimNextUncookedOnly/Internal/AnimNextRigVMAssetEditorData.h
@@ -189,7 +189,6 @@ protected:
 	virtual void Serialize(FArchive& Ar) override;
 	virtual void PostLoad() override;
 	virtual void PostTransacted(const FTransactionObjectEvent& TransactionEvent) override;
-	virtual void PostDuplicate(EDuplicateMode::Type DuplicateMode) override;
 	virtual void GetAssetRegistryTags(FAssetRegistryTagsContext Context) const override;
 	virtual bool IsEditorOnly() const override { return true; }
 	virtual bool Rename(const TCHAR* NewName = nullptr, UObject* NewOuter = nullptr, ERenameFlags Flags = REN_None) override;	
diff --git a/Engine/Plugins/Experimental/AnimNext/Source/AnimNextUncookedOnly/Private/AnimNextRigVMAssetEditorData.cpp b/Engine/Plugins/Experimental/AnimNext/Source/AnimNextUncookedOnly/Private/AnimNextRigVMAssetEditorData.cpp
--- a/Engine/Plugins/Experimental/AnimNext/Source/AnimNextUncookedOnly/Private/AnimNextRigVMAssetEditorData.cpp
+++ b/Engine/Plugins/Experimental/AnimNext/Source/AnimNextUncookedOnly/Private/AnimNextRigVMAssetEditorData.cpp
@@ -213,7 +213,12 @@ void UAnimNextRigVMAssetEditorData::PostLoad()
 		BuildData->ClearInvalidReferences();
 	}
 
-	// delay compilation until the package has been loaded
+	// Mark this as being dirty so that we recompile when needed
+	bVMRecompilationRequired = true;
+
+	// Queue compilation once the package has been fully loaded
+	// This is necessary in case we have external packages that haven't post-loaded yet
+	// However, if we are duplicating the asset OnEndLoadPackage won't be called
 	FCoreUObjectDelegates::OnEndLoadPackage.AddUObject(this, &UAnimNextRigVMAssetEditorData::HandlePackageDone);
 }
 
@@ -246,11 +251,6 @@ void UAnimNextRigVMAssetEditorData::PostTransacted(const FTransactionObjectEvent
 	}
 }
 
-void UAnimNextRigVMAssetEditorData::PostDuplicate(EDuplicateMode::Type DuplicateMode)
-{
-	Initialize(/*bRecompileVM*/true);
-}
-
 void UAnimNextRigVMAssetEditorData::GetAssetRegistryTags(FAssetRegistryTagsContext Context) const
 {
 	Super::GetAssetRegistryTags(Context);
