AnimNext: Injection

Adds new 'injection site' concept, accompanying trait and infrastructure.
An injection site is a variable describing what to inject, linked to a trait that tells the system to inject the trait int the specific location in the graph.
'Run Graph' nodes now do not host the graph to run, rather they host a small graph that contains a blend stack and an injection site. This allows top-level graphs in modules to be 'injected' into with blending, rather than popping.
Multi-level binding propagation is required to make all this work, as bindings previously only worked across adjacent 'layers' from host -> client graphs, and the changes above to Run Graph nodes imply at least 3 adjacent 'layers' host -> client -> injected. To this end, bindings are now cached, and automatic bindings are now a combination of cached outer bindings and any inner bindings (inner overriding outer).
Injection has a few modes:
- Regular graph injection - this takes an asset and factory-creates a graph to run it. This all occurs mostly in the blend stack node.
- Regular graph injection but with external module bindings. This allows graphs to be injected and their automatic bindings to be bound to both their host AND an external module. In this way, modules that host graphs do not need to know about the concerns of the graph's interface (although the graph can still pull values from the host as well as the external module).
- Task evaluation injection - this is a low-level injection system that allows systems like Sequencer to inject their own tasks into the FEvaluationProgram during the IEvaluate traversal.

The addition of external module bindings has made it necessaary to add MT access detectors to prevent race conditions when running one module in the context of another's data. Tick dependencies are set up on injection to ensure that tasks do not run concurrently and the MT access detector does not fire.

#virtualized

diff --git a/Engine/Plugins/Experimental/AnimNext/Config/BaseAnimNext.ini b/Engine/Plugins/Experimental/AnimNext/Config/BaseAnimNext.ini
--- a/Engine/Plugins/Experimental/AnimNext/Config/BaseAnimNext.ini
+++ b/Engine/Plugins/Experimental/AnimNext/Config/BaseAnimNext.ini
@@ -86,7 +86,6 @@
 +StructRedirects=(OldName="/Script/AnimNext.AnimNextTraitHandle", NewName="/Script/AnimNextAnimGraph.AnimNextTraitHandle")
 +StructRedirects=(OldName="/Script/AnimNext.AnimNextTraitSharedData", NewName="/Script/AnimNextAnimGraph.AnimNextTraitSharedData")
 +StructRedirects=(OldName="/Script/AnimNext.AnimNextGraphInstance", NewName="/Script/AnimNextAnimGraph.AnimNextGraphInstance")
-+StructRedirects=(OldName="/Script/AnimNext.AnimNextGraphInstancePtr", NewName="/Script/AnimNextAnimGraph.AnimNextGraphInstancePtr")
 
 +ClassRedirects=(OldName="/Script/AnimNextEditor.AnimNextAnimationGraphFactory", NewName="/Script/AnimNextAnimGraphEditor.AnimNextAnimationGraphFactory")
 +ClassRedirects=(OldName="/Script/AnimNextEditor.AssetDefinition_AnimNextAnimationGraph", NewName="/Script/AnimNextAnimGraphEditor.AssetDefinition_AnimNextAnimationGraph")
@@ -96,8 +95,9 @@
 +StructRedirects=(OldName="/Script/AnimNextUncookedOnly.AnimNextAnimationGraphOutlinerData", NewName="/Script/AnimNextAnimGraphUncookedOnly.AnimNextAnimationGraphOutlinerData")
 +ClassRedirects=(OldName="/Script/AnimNextUncookedOnly.AnimNextAnimationGraphEntry", NewName="/Script/AnimNextAnimGraphUncookedOnly.AnimNextAnimationGraphEntry")
 
-
-
+; PlayAnim -> Injection
++StructRedirects=(OldName="/Script/AnimNext.ModuleInstanceComponent", NewName="/Script/AnimNext.AnimNextModuleInstanceComponent")
++StructRedirects=(OldName="/Script/AnimNext.AnimNextGraphInstancePtr", NewName="/Script/AnimNextAnimGraph.AnimNextAnimGraph")
 
 
 
diff --git a/Engine/Plugins/Experimental/AnimNext/Source/AnimNext/Internal/AnimNextPool.h b/Engine/Plugins/Experimental/AnimNext/Source/AnimNext/Internal/AnimNextPool.h
--- a/Engine/Plugins/Experimental/AnimNext/Source/AnimNext/Internal/AnimNextPool.h
+++ b/Engine/Plugins/Experimental/AnimNext/Source/AnimNext/Internal/AnimNextPool.h
@@ -5,34 +5,11 @@
 #include "CoreTypes.h"
 #include "Containers/PagedArray.h"
 #include "Misc/TVariant.h"
+#include "AnimNextPoolHandle.h"
 
 namespace UE::AnimNext
 {
 
-template<typename ElementType>
-class TPool;
-
-template<typename ElementType>
-struct TPoolHandle
-{
-private:
-	uint32 Index = MAX_uint32;
-	uint32 SerialNumber = 0;
-
-	friend TPool<ElementType>;
-public:
-	void Reset()
-	{
-		Index = MAX_uint32;
-		SerialNumber = 0;
-	}
-
-	bool IsValid() const
-	{
-		return Index != MAX_uint32 && SerialNumber != 0;
-	}
-};
-
 template<typename ElementType>
 class TPool
 {
@@ -117,8 +94,17 @@ public:
 		check(IsValidHandle(InHandle));
 		return Entries[InHandle.Index].ValueOrNextFreeIndex.template Get<ElementType>();
 	}
-	
-	ElementType* TryGet(TPoolHandle<ElementType> InHandle) const
+
+	ElementType* TryGet(TPoolHandle<ElementType> InHandle)
+	{
+		if(IsValidHandle(InHandle))
+		{
+			return &Entries[InHandle.Index].ValueOrNextFreeIndex.template Get<ElementType>();
+		}
+		return nullptr;
+	}
+
+	const ElementType* TryGet(TPoolHandle<ElementType> InHandle) const
 	{
 		if(IsValidHandle(InHandle))
 		{
diff --git a/Engine/Plugins/Experimental/AnimNext/Source/AnimNext/Internal/Component/AnimNextWorldSubsystem.h b/Engine/Plugins/Experimental/AnimNext/Source/AnimNext/Internal/Component/AnimNextWorldSubsystem.h
--- a/Engine/Plugins/Experimental/AnimNext/Source/AnimNext/Internal/Component/AnimNextWorldSubsystem.h
+++ b/Engine/Plugins/Experimental/AnimNext/Source/AnimNext/Internal/Component/AnimNextWorldSubsystem.h
@@ -17,6 +17,7 @@ struct FRigUnit_AnimNextInitializeEvent;
 namespace UE::AnimNext
 {
 	struct FModuleEndTickFunction;
+	struct FInjectionRequest;
 }
 
 namespace UE::AnimNext
@@ -116,12 +117,26 @@ protected:
 	// @param	InLocation			Where to run the task, before or after
 	void QueueTaskHandle(UE::AnimNext::FModuleHandle InHandle, FName InModuleEventName, TUniqueFunction<void(const UE::AnimNext::FModuleTaskContext&)>&& InTaskFunction, UE::AnimNext::ETaskRunLocation InLocation = UE::AnimNext::ETaskRunLocation::Before);
 
-	// Disables the initialize event for a module one it has run once
+	// Queue an input trait event on the supplied handle
+	void QueueInputTraitEventHandle(UE::AnimNext::FModuleHandle InHandle, FAnimNextTraitEventPtr Event);
+
+	// Add a dependency between equivalent module events for the specified handles
+	// @param	InPrerequisiteHandle	The handle of the module that will be executed first
+	// @param	InSubsequentHandle		The handle of the module that will be executed second
+	void AddModuleDependencyHandle(UE::AnimNext::FModuleHandle InPrerequisiteHandle, UE::AnimNext::FModuleHandle InSubsequentHandle);
+
+	// Remove a dependency between equivalent module events for the specified handles
+	// @param	InPrerequisiteHandle	The handle of the module that would have been executed first
+	// @param	InSubsequentHandle		The handle of the module that would have been executed second
+	void RemoveModuleDependencyHandle(UE::AnimNext::FModuleHandle InPrerequisiteHandle, UE::AnimNext::FModuleHandle InSubsequentHandle);
+
+	// Disables the initialize event for a module once it has run once
 	void DisableInitializeEvent(UE::AnimNext::FModuleHandle InHandle);
-	
+
 protected:
 	friend struct FRigUnit_AnimNextInitializeEvent;
 	friend struct UE::AnimNext::FModuleEndTickFunction;
+	friend struct UE::AnimNext::FInjectionRequest;
 
 	// Currently running instances, pooled
 	UE::AnimNext::TPool<FAnimNextModuleInstance> Instances;
diff --git a/Engine/Plugins/Experimental/AnimNext/Source/AnimNext/Internal/Module/AnimNextModule.h b/Engine/Plugins/Experimental/AnimNext/Source/AnimNext/Internal/Module/AnimNextModule.h
--- a/Engine/Plugins/Experimental/AnimNext/Source/AnimNext/Internal/Module/AnimNextModule.h
+++ b/Engine/Plugins/Experimental/AnimNext/Source/AnimNext/Internal/Module/AnimNextModule.h
@@ -15,7 +15,6 @@ class UAnimNextModule;
 class UAnimGraphNode_AnimNextGraph;
 struct FAnimNode_AnimNextGraph;
 struct FRigUnit_AnimNextGraphEvaluator;
-struct FAnimNextGraphInstancePtr;
 struct FAnimNextGraphInstance;
 struct FAnimNextScheduleGraphTask;
 struct FAnimNextEditorParam;
@@ -41,22 +40,6 @@ namespace UE::AnimNext::Editor
 	class FVariableCustomization;
 }
 
-UENUM()
-enum class EAnimNextModuleInitMethod : uint8
-{
-	// Do not perform any initial update, set up data structures only
-	None,
-
-	// Set up data structures, perform an initial update and then pause
-	InitializeAndPause,
-
-	// Set up data structures, perform an initial update and then pause in editor only, otherwise act like InitializeAndRun
-	InitializeAndPauseInEditor,
-
-	// Set up data structures then continue updating
-	InitializeAndRun
-};
-
 // Root asset represented by a component when instantiated
 UCLASS(BlueprintType)
 class ANIMNEXT_API UAnimNextModule : public UAnimNextDataInterface
@@ -82,7 +65,6 @@ protected:
 	friend struct UE::AnimNext::UncookedOnly::FUtils;
 	friend class UE::AnimNext::Editor::FModuleEditor;
 	friend struct UE::AnimNext::FTestUtils;
-	friend FAnimNextGraphInstancePtr;
 	friend FAnimNextGraphInstance;
 	friend class UAnimGraphNode_AnimNextGraph;
 	friend UE::AnimNext::FExecutionContext;
diff --git a/Engine/Plugins/Experimental/AnimNext/Source/AnimNext/Internal/Module/AnimNextModuleInstance.h b/Engine/Plugins/Experimental/AnimNext/Source/AnimNext/Internal/Module/AnimNextModuleInstance.h
--- a/Engine/Plugins/Experimental/AnimNext/Source/AnimNext/Internal/Module/AnimNextModuleInstance.h
+++ b/Engine/Plugins/Experimental/AnimNext/Source/AnimNext/Internal/Module/AnimNextModuleInstance.h
@@ -4,7 +4,7 @@
 
 #include "CoreMinimal.h"
 #include "AnimNextModule.h"
-#include "ModuleHandle.h"
+#include "Module/ModuleHandle.h"
 #include "Component/AnimNextPublicVariablesProxy.h"
 #include "DataInterface/AnimNextDataInterfaceInstance.h"
 #include "Module/ModuleTickFunction.h"
@@ -13,20 +13,26 @@
 #include "Module/ModuleHandle.h"
 #include "Module/ModuleInstanceComponent.h"
 #include "StructUtils/InstancedStruct.h"
+#include "Module/AnimNextModuleInitMethod.h"
 
 #include "AnimNextModuleInstance.generated.h"
 
 struct FAnimNextModuleInstance;
-struct FRigUnit_AnimNextRunAnimationGraph;
+struct FRigUnit_AnimNextRunAnimationGraph_v1;
+struct FRigUnit_AnimNextRunAnimationGraph_v2;
 class UAnimNextWorldSubsystem;
 struct FRigUnit_AnimNextInitializeEvent;
 struct FRigUnit_CopyModuleProxyVariables;
+struct FAnimNextNotifyDispatcherComponent;
+struct FAnimNextGraphInstance;
 
 namespace UE::AnimNext
 {
 	struct FProxyVariablesContext;
 	struct FModuleEventTickFunction;
 	struct FModuleEndTickFunction;
+	struct FModuleInjectionDataInterfaceAdapter;
+	struct FModuleWriteGuard;
 }
 
 namespace UE::AnimNext::Debug
@@ -60,7 +66,7 @@ namespace UE::AnimNext::Private
 
 namespace UE::AnimNext
 {
-	using ModuleInstanceComponentMapType = TMap<FName, TInstancedStruct<FModuleInstanceComponent>>;
+	using ModuleInstanceComponentMapType = TMap<FName, TInstancedStruct<FAnimNextModuleInstanceComponent>>;
 }
 
 // Root memory owner of a parameterized schedule 
@@ -73,6 +79,7 @@ struct FAnimNextModuleInstance : public FAnimNextDataInterfaceInstance
 	ANIMNEXT_API FAnimNextModuleInstance(
 		UAnimNextModule* InModule,
 		UObject* InObject,
+		UE::AnimNext::TPool<FAnimNextModuleInstance>* InPool,
 		EAnimNextModuleInitMethod InInitMethod);
 	ANIMNEXT_API ~FAnimNextModuleInstance();
 
@@ -106,8 +113,8 @@ struct FAnimNextModuleInstance : public FAnimNextDataInterfaceInstance
 	const ComponentType* TryGetComponent() const;
 
 	// Returns a pointer to the specified component, or nullptr if not found
-	ANIMNEXT_API const FModuleInstanceComponent* TryGetComponent(int32 ComponentNameHash, FName ComponentName) const;
-	ANIMNEXT_API FModuleInstanceComponent* TryGetComponent(int32 ComponentNameHash, FName ComponentName);
+	ANIMNEXT_API const FAnimNextModuleInstanceComponent* TryGetComponent(int32 ComponentNameHash, FName ComponentName) const;
+	ANIMNEXT_API FAnimNextModuleInstanceComponent* TryGetComponent(int32 ComponentNameHash, FName ComponentName);
 
 #if UE_ENABLE_DEBUG_DRAWING
 	// Get the debug draw interface
@@ -132,8 +139,8 @@ struct FAnimNextModuleInstance : public FAnimNextDataInterfaceInstance
 	// Get the world type that this module was instantiated within
 	EWorldType::Type GetWorldType() const { return WorldType; }
 
-	// Get the world type that this module was instantiated within
-	TArrayView<UE::AnimNext::FModuleEventTickFunction> GetTickFunctions() { return TickFunctions; }
+	// Get tick functions for this module instance
+	ANIMNEXT_API TArrayView<UE::AnimNext::FModuleEventTickFunction> GetTickFunctions();
 
 private:
 	// Setup the entry
@@ -142,12 +149,6 @@ private:
 	// Clear data that binds the schedule to a runtime (e.g. tick functions) and any instance data
 	void ResetBindingsAndInstanceData();
 
-	// Used for pooling
-	void Invalidate();
-
-	// Clears the bTickEvenWhenPaused flags of the entries' tick functions
-	void ClearTickFunctionPauseFlags();
-
 	// Flip proxy variable buffers then copy any dirty values
 	void CopyProxyVariables();
 
@@ -159,26 +160,46 @@ private:
 	// Called to disable the initialize event for this module on the game thread when it has executed for the first time
 	void DisableInitializeEvent();
 
-	// Give the module components a chance to handle events and call their OnEndExecution extension point 
+	// Give the module components a chance to handle events and call their OnEndExecution extension point
 	void EndExecution(float InDeltaTime);
 
 	// Handles trait events at the module level.
 	// Called in the module end tick function before any events are dispatched.
 	void RaiseTraitEvents(const UE::AnimNext::FTraitEventList& EventList);
 
+	// Add a prerequisite instance
+	void AddPrerequisite(FAnimNextModuleInstance& InPrerequisiteInstance);
+
+	// Remove a prerequisite instance
+	void RemovePrerequisite(FAnimNextModuleInstance& InPrerequisiteInstance);
+
+	// @return true if the instance is a prerequisite instance
+	ANIMNEXT_API bool IsPrerequisite(const FAnimNextModuleInstance& InPrerequisiteInstance) const;
+
+	// Call the supplied function for each prerequisite module
+	ANIMNEXT_API void ForEachPrerequisite(TFunctionRef<void(FAnimNextModuleInstance& InPrerequisiteInstance)>) const;
+
 private:
 	friend UE::AnimNext::FModuleEventTickFunction;
 	friend UE::AnimNext::FModuleEndTickFunction;
-	friend FRigUnit_AnimNextRunAnimationGraph;
+	friend FRigUnit_AnimNextRunAnimationGraph_v1;
+	friend FRigUnit_AnimNextRunAnimationGraph_v2;
 	friend UAnimNextWorldSubsystem;
 	friend FRigUnit_AnimNextInitializeEvent;
 	friend UE::AnimNext::FProxyVariablesContext;
 	friend FRigUnit_CopyModuleProxyVariables;
+	friend UE::AnimNext::FModuleInjectionDataInterfaceAdapter;
+	friend UE::AnimNext::FModuleWriteGuard;
+	friend FAnimNextNotifyDispatcherComponent;
+	friend FAnimNextGraphInstance;
 
 	// Object this entry is bound to
 	UPROPERTY(Transient)
 	TObjectPtr<UObject> Object = nullptr;
 
+	// The pool that this module instance exists in
+	UE::AnimNext::TPool<FAnimNextModuleInstance>* Pool = nullptr;
+
 	// Copy of the handle that represents this entry to client systems
 	UE::AnimNext::FModuleHandle Handle;
 
@@ -188,12 +209,26 @@ private:
 	// Pre-allocated graph of tick functions
 	TArray<UE::AnimNext::FModuleEventTickFunction> TickFunctions;
 
+	// Reference-counted prerequisites. Prerequisites can be requested/unrequested multiple times, so are only truly removed when references are zero
+	struct FPrerequisiteReference
+	{
+		UE::AnimNext::FModuleHandle Handle;
+		int32 ReferenceCount = 0;
+	};
+
+	// All pre-requisite references, only modified from the game thread
+	TArray<FPrerequisiteReference> PrerequisiteRefs;
+
 	// Input event list to be processed on the next update
 	UE::AnimNext::FTraitEventList InputEventList;
 
 	// Output event list to be processed at the end of the schedule tick
 	UE::AnimNext::FTraitEventList OutputEventList;
 
+	// Access detector to ensure module prerequisite and general execution access is safe
+	// This protects all calls & tick functions and prevents prerequisites from running unrelated tick functions concurrently
+	UE_MT_DECLARE_RW_RECURSIVE_ACCESS_DETECTOR(AccessDetector);
+
 	// Lock to ensure event list actions are thread safe
 	FRWLock EventListLock;
 
@@ -209,7 +244,7 @@ private:
 
 	// Components of this module
 	UPROPERTY()
-	TMap<FName, TInstancedStruct<FModuleInstanceComponent>> ComponentMap;
+	TMap<FName, TInstancedStruct<FAnimNextModuleInstanceComponent>> ComponentMap;
 
 	// Initialize event index (INDEX_NONE if no initialize event is found)
 	int32 InitializeEventIndex = INDEX_NONE;
@@ -248,7 +283,7 @@ private:
 
 private:
 	// Adds the specified component and returns a reference to it
-	ANIMNEXT_API FModuleInstanceComponent& AddComponentInternal(int32 ComponentNameHash, FName ComponentName, TInstancedStruct<FModuleInstanceComponent>&& Component);
+	ANIMNEXT_API FAnimNextModuleInstanceComponent& AddComponentInternal(int32 ComponentNameHash, FName ComponentName, TInstancedStruct<FAnimNextModuleInstanceComponent>&& Component);
 };
 
 template<>
diff --git a/Engine/Plugins/Experimental/AnimNext/Source/AnimNext/Internal/Module/ModuleGuard.h b/Engine/Plugins/Experimental/AnimNext/Source/AnimNext/Internal/Module/ModuleGuard.h
new file mode 100644
--- /dev/null
+++ b/Engine/Plugins/Experimental/AnimNext/Source/AnimNext/Internal/Module/ModuleGuard.h
@@ -0,0 +1,26 @@
+﻿// Copyright Epic Games, Inc. All Rights Reserved.
+
+#pragma once
+
+#include "CoreMinimal.h"
+#include "Module/ModuleHandle.h"
+
+struct FAnimNextModuleInstance;
+
+namespace UE::AnimNext
+{
+
+// RAII helper that guards against concurrent write access to the specified module and all its prerequisites
+struct FModuleWriteGuard
+{
+	UE_NONCOPYABLE(FModuleWriteGuard);
+
+	ANIMNEXT_API explicit FModuleWriteGuard(const FAnimNextModuleInstance* InModuleInstance);
+	ANIMNEXT_API ~FModuleWriteGuard();
+
+private:
+	const FAnimNextModuleInstance* ModuleInstance = nullptr;
+	TArray<FModuleHandle, TInlineAllocator<4>> PrerequisiteHandles;
+};
+
+}
diff --git a/Engine/Plugins/Experimental/AnimNext/Source/AnimNext/Internal/Module/ModuleInstanceComponent.h b/Engine/Plugins/Experimental/AnimNext/Source/AnimNext/Internal/Module/ModuleInstanceComponent.h
--- a/Engine/Plugins/Experimental/AnimNext/Source/AnimNext/Internal/Module/ModuleInstanceComponent.h
+++ b/Engine/Plugins/Experimental/AnimNext/Source/AnimNext/Internal/Module/ModuleInstanceComponent.h
@@ -14,21 +14,21 @@ struct FAnimNextTraitEvent;
 	virtual FName GetComponentName() const override { return StaticComponentName(); }
 
 /**
- * FModuleInstanceComponent
+ * FAnimNextModuleInstanceComponent
  *
  * A module instance component is attached and owned by a module instance.
  * It persists as long as it is needed.
  */
- USTRUCT()
-struct FModuleInstanceComponent
+USTRUCT()
+struct FAnimNextModuleInstanceComponent
 {
 	GENERATED_BODY()
 
-	FModuleInstanceComponent() = default;
-	explicit FModuleInstanceComponent(FAnimNextModuleInstance& InOwnerInstance) : OwnerInstance(&InOwnerInstance) {}
-	virtual ~FModuleInstanceComponent() {}
+	FAnimNextModuleInstanceComponent() = default;
+	explicit FAnimNextModuleInstanceComponent(FAnimNextModuleInstance& InOwnerInstance) : OwnerInstance(&InOwnerInstance) {}
+	virtual ~FAnimNextModuleInstanceComponent() {}
 
-	static FName StaticComponentName() { return FName(TEXT("FModuleInstanceComponent")); }
+	static FName StaticComponentName() { return FName(TEXT("FAnimNextModuleInstanceComponent")); }
 	virtual FName GetComponentName() const { return StaticComponentName(); }
 
 	// Returns the owning module instance this component lives on
@@ -37,10 +37,10 @@ struct FModuleInstanceComponent
 	// Returns the owning module instance this component lives on
 	const FAnimNextModuleInstance& GetModuleInstance() const { return *OwnerInstance; }
 
-	// Called at the end of module execution for any output events to be handled, before OnEndExecution
+	// Called during module execution for any events to be handled
 	virtual void OnTraitEvent(FAnimNextTraitEvent& Event) {}
 
-	// Called at end of module execution
+	// Called at end of module execution each frame
 	virtual void OnEndExecution(float InDeltaTime) {}
 
 private:
diff --git a/Engine/Plugins/Experimental/AnimNext/Source/AnimNext/Internal/Module/ModuleTaskContext.h b/Engine/Plugins/Experimental/AnimNext/Source/AnimNext/Internal/Module/ModuleTaskContext.h
--- a/Engine/Plugins/Experimental/AnimNext/Source/AnimNext/Internal/Module/ModuleTaskContext.h
+++ b/Engine/Plugins/Experimental/AnimNext/Source/AnimNext/Internal/Module/ModuleTaskContext.h
@@ -8,6 +8,7 @@
 
 struct FAnimNextModuleInstance;
 struct FInstancedPropertyBag;
+struct FAnimNextModuleInjectionComponent;
 
 namespace UE::AnimNext
 {
@@ -41,8 +42,9 @@ private:
 
 	// The module instance currently running
 	FAnimNextModuleInstance* ModuleInstance;
-	
-	friend struct FModuleEventTickFunction;
+
+	friend UE::AnimNext::FModuleEventTickFunction;
+	friend ::FAnimNextModuleInjectionComponent;
 };
 
 }
\ No newline at end of file
diff --git a/Engine/Plugins/Experimental/AnimNext/Source/AnimNext/Internal/Module/ModuleTickFunction.h b/Engine/Plugins/Experimental/AnimNext/Source/AnimNext/Internal/Module/ModuleTickFunction.h
--- a/Engine/Plugins/Experimental/AnimNext/Source/AnimNext/Internal/Module/ModuleTickFunction.h
+++ b/Engine/Plugins/Experimental/AnimNext/Source/AnimNext/Internal/Module/ModuleTickFunction.h
@@ -17,6 +17,9 @@ namespace UE::AnimNext
 namespace UE::AnimNext
 {
 
+// Multicast delegate called before module events
+using FOnPreModuleEvent = TTSMulticastDelegate<void(const FModuleTaskContext& InContext)>;
+
 struct FModuleEventTickFunction : public FTickFunction
 {
 	FModuleEventTickFunction()
@@ -41,9 +44,17 @@ struct FModuleEventTickFunction : public FTickFunction
 #endif
 
 	FAnimNextModuleInstance* ModuleInstance = nullptr;
-	FName EventName;
+
+	// Tasks queued to execute once per event
 	TSpscQueue<TUniqueFunction<void(const FModuleTaskContext&)>> PreExecuteTasks;
 	TSpscQueue<TUniqueFunction<void(const FModuleTaskContext&)>> PostExecuteTasks;
+
+	// Multicast delegate called before this event
+	FOnPreModuleEvent OnPreModuleEvent;
+
+	FName EventName;
+
+	bool bUserEvent = false;
 };
 
 struct FModuleEndTickFunction : public FTickFunction
diff --git a/Engine/Plugins/Experimental/AnimNext/Source/AnimNext/Internal/Module/RigUnit_AnimNextModuleEvents.h b/Engine/Plugins/Experimental/AnimNext/Source/AnimNext/Internal/Module/RigUnit_AnimNextModuleEvents.h
--- a/Engine/Plugins/Experimental/AnimNext/Source/AnimNext/Internal/Module/RigUnit_AnimNextModuleEvents.h
+++ b/Engine/Plugins/Experimental/AnimNext/Source/AnimNext/Internal/Module/RigUnit_AnimNextModuleEvents.h
@@ -31,15 +31,34 @@ struct ANIMNEXT_API FRigUnit_AnimNextModuleEventBase : public FRigUnit_AnimNextB
 	FAnimNextExecuteContext ExecuteContext;
 };
 
-/** Synthetic event injected by the compiler to process any variable bindings, not user instantiated */
+/** Synthetic event injected by the compiler to process any variable bindings on the game thread, not user instantiated */
 USTRUCT(meta=(Hidden, Category="Internal", NodeColor="1, 0, 0"))
-struct ANIMNEXT_API FRigUnit_AnimNextExecuteBindings : public FRigUnit_AnimNextModuleEventBase
+struct ANIMNEXT_API FRigUnit_AnimNextExecuteBindings_GT : public FRigUnit_AnimNextModuleEventBase
 {
 	GENERATED_BODY()
 
-	FRigUnit_AnimNextExecuteBindings() = default;
+	FRigUnit_AnimNextExecuteBindings_GT() = default;
 
-	static inline const FLazyName EventName = FLazyName("ExecuteBindings");
+	static inline const FLazyName EventName = FLazyName("ExecuteBindings_GT");
+	
+	RIGVM_METHOD()
+	virtual void Execute() override;
+	virtual FName GetEventName() const override final { return EventName; }
+
+	// FRigUnit_AnimNextModuleEventBase interface
+	virtual UE::AnimNext::EModuleEventPhase GetEventPhase() const override final { return UE::AnimNext::EModuleEventPhase::PreExecute; }
+	virtual UE::AnimNext::FModuleEventBindingFunction GetBindingFunction() const override final;
+};
+
+/** Synthetic event injected by the compiler to process any variable bindings on a worker thread, not user instantiated */
+USTRUCT(meta=(Hidden, Category="Internal", NodeColor="1, 0, 0"))
+struct ANIMNEXT_API FRigUnit_AnimNextExecuteBindings_WT : public FRigUnit_AnimNextModuleEventBase
+{
+	GENERATED_BODY()
+
+	FRigUnit_AnimNextExecuteBindings_WT() = default;
+
+	static inline const FLazyName EventName = FLazyName("ExecuteBindings_WT");
 	
 	RIGVM_METHOD()
 	virtual void Execute() override;
diff --git a/Engine/Plugins/Experimental/AnimNext/Source/AnimNext/Private/Component/AnimNextComponent.cpp b/Engine/Plugins/Experimental/AnimNext/Source/AnimNext/Private/Component/AnimNextComponent.cpp
--- a/Engine/Plugins/Experimental/AnimNext/Source/AnimNext/Private/Component/AnimNextComponent.cpp
+++ b/Engine/Plugins/Experimental/AnimNext/Source/AnimNext/Private/Component/AnimNextComponent.cpp
@@ -882,10 +882,6 @@ void UAnimNextComponent::QueueInputTraitEvent(FAnimNextTraitEventPtr Event)
 
 	if(Subsystem)
 	{
-		Subsystem->QueueTask(this, NAME_None, [Event = MoveTemp(Event)](const FModuleTaskContext& InContext)
-		{
-			InContext.QueueInputTraitEvent(Event);
-		},
-		ETaskRunLocation::After);
+		Subsystem->QueueInputTraitEvent(this, Event);
 	}
 }
diff --git a/Engine/Plugins/Experimental/AnimNext/Source/AnimNext/Private/Component/AnimNextWorldSubsystem.cpp b/Engine/Plugins/Experimental/AnimNext/Source/AnimNext/Private/Component/AnimNextWorldSubsystem.cpp
--- a/Engine/Plugins/Experimental/AnimNext/Source/AnimNext/Private/Component/AnimNextWorldSubsystem.cpp
+++ b/Engine/Plugins/Experimental/AnimNext/Source/AnimNext/Private/Component/AnimNextWorldSubsystem.cpp
@@ -7,12 +7,8 @@
 #include "Module/AnimNextModule.h"
 #include "Module/ModuleHandle.h"
 #include "Module/ModuleTaskContext.h"
-
-namespace UE::AnimNext
-{
-
-
-}
+#include "AnimNextDebugDraw.h"
+#include "Module/ModuleGuard.h"
 
 UAnimNextWorldSubsystem::UAnimNextWorldSubsystem()
 {
@@ -131,6 +127,7 @@ void UAnimNextWorldSubsystem::FlushPendingActions()
 #endif
 				break;
 			default:
+				checkNoEntry();
 				break;
 			}
 		}
@@ -143,9 +140,9 @@ void UAnimNextWorldSubsystem::RegisterHandle(UE::AnimNext::FModuleHandle& InOutH
 {
 	using namespace UE::AnimNext;
 	check(IsInGameThread());
-	FRWScopeLock InstancesLockScope(InstancesLock, SLT_Write);
+	FWriteScopeLock InstancesLockScope(InstancesLock);
 
-	InOutHandle = Instances.Emplace(InModule, InObject, InitMethod);
+	InOutHandle = Instances.Emplace(InModule, InObject, &Instances, InitMethod);
 	FAnimNextModuleInstance& Instance = Instances.Get(InOutHandle);
 	Instance.Handle = InOutHandle;
 	Instance.Initialize();
@@ -158,13 +155,18 @@ void UAnimNextWorldSubsystem::UnregisterHandle(UE::AnimNext::FModuleHandle& InOu
 
 	if(IsValidHandle(InOutHandle))
 	{
-		FRWScopeLock InstancesLockScope(InstancesLock, SLT_Write);
+#if UE_ENABLE_DEBUG_DRAWING
+		// Release debug drawing immediately as the renderer will need to know about this before EOF
+		{
+			FWriteScopeLock InstancesLockScope(InstancesLock);
+			FAnimNextModuleInstance& Instance = Instances.Get(InOutHandle);
+			FModuleWriteGuard Guard(&Instance);
+			Instance.DebugDraw.Reset();
+		}
+#endif
 
-		Instances.Release(InOutHandle);
+		PendingActions.Emplace(InOutHandle, FModulePendingAction::EType::ReleaseHandle);
 		InOutHandle.Reset();
-
-		// TODO: do not allow immediate handle releases outside of schedule runs - we can either defer or assert
-	//	PendingActions.Emplace(FModulePendingAction::EType::ReleaseHandle, InOutHandle);
 	}
 }
 
@@ -267,4 +269,53 @@ void UAnimNextWorldSubsystem::DisableInitializeEvent(UE::AnimNext::FModuleHandle
 		FAnimNextModuleInstance& Instance = Instances.Get(InHandle);
 		Instance.DisableInitializeEvent();
 	}
+}
+
+void UAnimNextWorldSubsystem::QueueInputTraitEventHandle(UE::AnimNext::FModuleHandle InHandle, FAnimNextTraitEventPtr Event)
+{
+	using namespace UE::AnimNext;
+
+	QueueTaskHandle(InHandle, NAME_None, [Event = MoveTemp(Event)](const FModuleTaskContext& InContext)
+	{
+		InContext.QueueInputTraitEvent(Event);
+	},
+	ETaskRunLocation::Before);
+}
+
+void UAnimNextWorldSubsystem::AddModuleDependencyHandle(UE::AnimNext::FModuleHandle InPrerequisiteHandle, UE::AnimNext::FModuleHandle InSubsequentHandle)
+{
+	using namespace UE::AnimNext;
+	check(IsInGameThread());
+	if (!IsValidHandle(InPrerequisiteHandle) || !IsValidHandle(InSubsequentHandle))
+	{
+		return;
+	}
+
+	FAnimNextModuleInstance* PrerequisiteInstance = Instances.TryGet(InPrerequisiteHandle);
+	FAnimNextModuleInstance* SubsequentInstance = Instances.TryGet(InSubsequentHandle);
+	if(PrerequisiteInstance == nullptr || SubsequentInstance == nullptr)
+	{
+		return;
+	}
+
+	SubsequentInstance->AddPrerequisite(*PrerequisiteInstance);
+}
+
+void UAnimNextWorldSubsystem::RemoveModuleDependencyHandle(UE::AnimNext::FModuleHandle InPrerequisiteHandle, UE::AnimNext::FModuleHandle InSubsequentHandle)
+{
+	using namespace UE::AnimNext;
+	check(IsInGameThread());
+	if (!IsValidHandle(InPrerequisiteHandle) || !IsValidHandle(InSubsequentHandle))
+	{
+		return;
+	}
+
+	FAnimNextModuleInstance* PrerequisiteInstance = Instances.TryGet(InPrerequisiteHandle);
+	FAnimNextModuleInstance* SubsequentInstance = Instances.TryGet(InSubsequentHandle);
+	if(PrerequisiteInstance == nullptr || SubsequentInstance == nullptr)
+	{
+		return;
+	}
+
+	SubsequentInstance->RemovePrerequisite(*PrerequisiteInstance);
 }
\ No newline at end of file
diff --git a/Engine/Plugins/Experimental/AnimNext/Source/AnimNext/Private/Module/AnimNextModuleInstance.cpp b/Engine/Plugins/Experimental/AnimNext/Source/AnimNext/Private/Module/AnimNextModuleInstance.cpp
--- a/Engine/Plugins/Experimental/AnimNext/Source/AnimNext/Private/Module/AnimNextModuleInstance.cpp
+++ b/Engine/Plugins/Experimental/AnimNext/Source/AnimNext/Private/Module/AnimNextModuleInstance.cpp
@@ -5,6 +5,7 @@
 #if UE_ENABLE_DEBUG_DRAWING
 #include "AnimNextDebugDraw.h"
 #endif
+#include "AnimNextPool.h"
 #include "Engine/World.h"
 #include "AnimNextStats.h"
 #include "Async/TaskGraphInterfaces.h"
@@ -12,6 +13,7 @@
 #include "Algo/TopologicalSort.h"
 #include "Logging/StructuredLog.h"
 #include "Misc/EnumerateRange.h"
+#include "Module/ModuleGuard.h"
 #include "Module/RigUnit_AnimNextModuleEvents.h"
 #include "Module/ProxyVariablesContext.h"
 #include "UObject/UObjectIterator.h"
@@ -24,11 +26,16 @@ FAnimNextModuleInstance::FAnimNextModuleInstance() = default;
 FAnimNextModuleInstance::FAnimNextModuleInstance(
 		UAnimNextModule* InModule,
 		UObject* InObject,
+		UE::AnimNext::TPool<FAnimNextModuleInstance>* InPool,
 		EAnimNextModuleInitMethod InInitMethod)
 	: Object(InObject)
+	, Pool(InPool)
 	, RunState(ERunState::None)
 	, InitMethod(InInitMethod)
 {
+	check(InModule);
+	check(InObject);
+
 	DataInterface = InModule;
 
 #if UE_ENABLE_DEBUG_DRAWING
@@ -61,6 +68,7 @@ struct FImplementedModuleEvent
 	FModuleEventBindingFunction Binding;
 	FName EventName;
 	EModuleEventPhase Phase = EModuleEventPhase::Execute;
+	bool bUserEvent = false;
 };
 
 static TArray<FImplementedModuleEvent> GAllModuleEvents;
@@ -84,6 +92,8 @@ void CacheAllModuleEvents()
 				NewEvent.Binding = StructInstance.Get().GetBindingFunction();
 				NewEvent.EventName = StructInstance.Get().GetEventName();
 				NewEvent.Phase = StructInstance.Get().GetEventPhase();
+				NewEvent.bUserEvent = Struct != FRigUnit_AnimNextExecuteBindings_GT::StaticStruct() &&
+									  Struct != FRigUnit_AnimNextExecuteBindings_WT::StaticStruct();
 			}
 		}
 
@@ -122,6 +132,8 @@ void FAnimNextModuleInstance::Initialize()
 
 	check(IsInGameThread());
 
+	FModuleWriteGuard Guard(this);
+
 	check(Object)
 	const UAnimNextModule* Module = GetModule();
 	check(Module);
@@ -149,6 +161,7 @@ void FAnimNextModuleInstance::Initialize()
 			FModuleEventTickFunction& TickFunction = TickFunctions[ModuleEvent.GetIndex()];
 			TickFunction.ModuleInstance = this;
 			TickFunction.EventName = ModuleEvent->EventName;
+			TickFunction.bUserEvent = ModuleEvent->bUserEvent;
 		}
 
 		// Bind tick functions (done in a second pass to allow functions to add prerequisites between each other)
@@ -245,6 +258,8 @@ void FAnimNextModuleInstance::ResetBindingsAndInstanceData()
 
 	check(IsInGameThread());
 
+	FModuleWriteGuard Guard(this);
+
 	TransitionToRunState(ERunState::None);
 
 	for (FModuleEventTickFunction& TickFunction : TickFunctions)
@@ -266,37 +281,17 @@ void FAnimNextModuleInstance::ResetBindingsAndInstanceData()
 	}
 }
 
-void FAnimNextModuleInstance::Invalidate()
-{
-	using namespace UE::AnimNext;
-
-	ResetBindingsAndInstanceData();
-
-	DataInterface = nullptr;
-	Object = nullptr;
-	Handle.Reset();
-}
-
-void FAnimNextModuleInstance::ClearTickFunctionPauseFlags()
-{
-	using namespace UE::AnimNext;
-
-	check(IsInGameThread());
-
-	for (FModuleEventTickFunction& TickFunction : TickFunctions)
-	{
-		TickFunction.bTickEvenWhenPaused = false;
-	}
-	EndTickFunction.bTickEvenWhenPaused = false;
-}
-
 void FAnimNextModuleInstance::QueueInputTraitEvent(FAnimNextTraitEventPtr Event)
 {
+	UE::AnimNext::FModuleWriteGuard Guard(this);
+
 	InputEventList.Push(MoveTemp(Event));
 }
 
 void FAnimNextModuleInstance::QueueOutputTraitEvent(FAnimNextTraitEventPtr Event)
 {
+	UE::AnimNext::FModuleWriteGuard Guard(this);
+
 	OutputEventList.Push(MoveTemp(Event));
 }
 
@@ -306,6 +301,8 @@ void FAnimNextModuleInstance::Enable(bool bInEnabled)
 
 	check(IsInGameThread());
 
+	FModuleWriteGuard Guard(this);
+
 	if(RunState == ERunState::PendingInitialUpdate || RunState == ERunState::Paused || RunState == ERunState::Running)
 	{
 		for (FModuleEventTickFunction& TickFunction : TickFunctions)
@@ -320,6 +317,8 @@ void FAnimNextModuleInstance::Enable(bool bInEnabled)
 
 void FAnimNextModuleInstance::TransitionToRunState(ERunState InNewState)
 {
+	UE::AnimNext::FModuleWriteGuard Guard(this);
+
 	switch(InNewState)
 	{
 	case ERunState::None:
@@ -349,6 +348,8 @@ void FAnimNextModuleInstance::TransitionToRunState(ERunState InNewState)
 
 void FAnimNextModuleInstance::CopyProxyVariables()
 {
+	UE::AnimNext::FModuleWriteGuard Guard(this);
+
 	// TODO: we can avoid the copies here by adopting a scheme where we:
 	// - Hold double-buffered memory handles
 	// - Update the memory handle's ptr to the currently-written double-buffered public variable on write
@@ -398,18 +399,22 @@ void FAnimNextModuleInstance::OnModuleCompiled()
 }
 #endif
 
-const FModuleInstanceComponent* FAnimNextModuleInstance::TryGetComponent(int32 ComponentNameHash, FName ComponentName) const
+const FAnimNextModuleInstanceComponent* FAnimNextModuleInstance::TryGetComponent(int32 ComponentNameHash, FName ComponentName) const
 {
-	if (const TInstancedStruct<FModuleInstanceComponent>* InstancedComponent = ComponentMap.FindByHash(ComponentNameHash, ComponentName))
+	UE::AnimNext::FModuleWriteGuard Guard(this);
+
+	if (const TInstancedStruct<FAnimNextModuleInstanceComponent>* InstancedComponent = ComponentMap.FindByHash(ComponentNameHash, ComponentName))
 	{
 		return InstancedComponent->GetPtr();
 	}
 
 	return nullptr;
 }
-FModuleInstanceComponent* FAnimNextModuleInstance::TryGetComponent(int32 ComponentNameHash, FName ComponentName) 
+FAnimNextModuleInstanceComponent* FAnimNextModuleInstance::TryGetComponent(int32 ComponentNameHash, FName ComponentName) 
 {
-	if (TInstancedStruct<FModuleInstanceComponent>* InstancedComponent = ComponentMap.FindByHash(ComponentNameHash, ComponentName))
+	UE::AnimNext::FModuleWriteGuard Guard(this);
+
+	if (TInstancedStruct<FAnimNextModuleInstanceComponent>* InstancedComponent = ComponentMap.FindByHash(ComponentNameHash, ComponentName))
 	{
 		return InstancedComponent->GetMutablePtr();
 	}
@@ -417,11 +422,13 @@ FModuleInstanceComponent* FAnimNextModuleInstance::TryGetComponent(int32 Compone
 	return nullptr;
 }
 
-FModuleInstanceComponent& FAnimNextModuleInstance::AddComponentInternal(int32 ComponentNameHash, FName ComponentName, TInstancedStruct<FModuleInstanceComponent>&& Component)
+FAnimNextModuleInstanceComponent& FAnimNextModuleInstance::AddComponentInternal(int32 ComponentNameHash, FName ComponentName, TInstancedStruct<FAnimNextModuleInstanceComponent>&& Component)
 {
-	TInstancedStruct<FModuleInstanceComponent>& InstancedComponent = ComponentMap.AddByHash(ComponentNameHash, ComponentName, Component);
+	UE::AnimNext::FModuleWriteGuard Guard(this);
+
+	TInstancedStruct<FAnimNextModuleInstanceComponent>& InstancedComponent = ComponentMap.AddByHash(ComponentNameHash, ComponentName, Component);
 	
-	FModuleInstanceComponent* ModuleInstanceComponent = InstancedComponent.GetMutablePtr();
+	FAnimNextModuleInstanceComponent* ModuleInstanceComponent = InstancedComponent.GetMutablePtr();
 	check(ModuleInstanceComponent != nullptr);
 
 	return *ModuleInstanceComponent;
@@ -430,11 +437,15 @@ FModuleInstanceComponent& FAnimNextModuleInstance::AddComponentInternal(int32 Co
 #if UE_ENABLE_DEBUG_DRAWING
 FRigVMDrawInterface* FAnimNextModuleInstance::GetDebugDrawInterface()
 {
+	UE::AnimNext::FModuleWriteGuard Guard(this);
+
 	return DebugDraw ? &DebugDraw->DrawInterface : nullptr;
 }
 
 void FAnimNextModuleInstance::ShowDebugDrawing(bool bInShowDebugDrawing)
 {
+	UE::AnimNext::FModuleWriteGuard Guard(this);
+
 	if(DebugDraw)
 	{
 		DebugDraw->SetEnabled(bInShowDebugDrawing);
@@ -490,7 +501,9 @@ void FAnimNextModuleInstance::GetModuleEventNames(TArray<FString>& OutNames, boo
 UE::AnimNext::FModuleEventTickFunction* FAnimNextModuleInstance::FindTickFunction(FName InEventName)
 {
 	using namespace UE::AnimNext;
-	
+
+	FModuleWriteGuard Guard(this);
+
 	for(FModuleEventTickFunction& TickFunction : TickFunctions)
 	{
 		if(TickFunction.EventName == InEventName)
@@ -503,22 +516,26 @@ UE::AnimNext::FModuleEventTickFunction* FAnimNextModuleInstance::FindTickFunctio
 
 void FAnimNextModuleInstance::EndExecution(float InDeltaTime)
 {
+	UE::AnimNext::FModuleWriteGuard Guard(this);
+
 	// Give the module a chance to handle events
 	RaiseTraitEvents(OutputEventList);
 
 	// Give each component a chance to finalize execution
 	for(auto It = ComponentMap.CreateIterator(); It; ++It)
 	{
-		FModuleInstanceComponent& Component = It->Value.GetMutable<FModuleInstanceComponent>();
+		FAnimNextModuleInstanceComponent& Component = It->Value.GetMutable<FAnimNextModuleInstanceComponent>();
 		Component.OnEndExecution(InDeltaTime);
 	}
 }
 
 void FAnimNextModuleInstance::RaiseTraitEvents(const UE::AnimNext::FTraitEventList& EventList)
 {
+	UE::AnimNext::FModuleWriteGuard Guard(this);
+
 	for(auto It = ComponentMap.CreateIterator(); It; ++It)
 	{
-		FModuleInstanceComponent& Component = It->Value.GetMutable<FModuleInstanceComponent>();
+		FAnimNextModuleInstanceComponent& Component = It->Value.GetMutable<FAnimNextModuleInstanceComponent>();
 
 		// Event handlers can raise events and as such the list may change while we iterate
 		// However, if an event is added while we iterate, we will not visit it
@@ -534,8 +551,98 @@ void FAnimNextModuleInstance::RaiseTraitEvents(const UE::AnimNext::FTraitEventLi
 	}
 }
 
+void FAnimNextModuleInstance::AddPrerequisite(FAnimNextModuleInstance& InPrerequisiteInstance)
+{
+	check(IsInGameThread());
+
+	UE::AnimNext::FModuleWriteGuard Guard(this);
+
+	FPrerequisiteReference* FoundReference = PrerequisiteRefs.FindByPredicate([&InPrerequisiteInstance](const FPrerequisiteReference& InPrerequisiteRef)
+	{
+		return InPrerequisiteRef.Handle == InPrerequisiteInstance.Handle;
+	});
+
+	if(FoundReference)
+	{
+		FoundReference->ReferenceCount++;
+	}
+	else
+	{
+		// Ensure all of our tick functions execute after the prerequisite's end tick function
+		EndTickFunction.AddPrerequisite(InPrerequisiteInstance.Object, InPrerequisiteInstance.EndTickFunction);
+		for(UE::AnimNext::FModuleEventTickFunction& TickFunction : TickFunctions)
+		{
+			TickFunction.AddPrerequisite(InPrerequisiteInstance.Object, InPrerequisiteInstance.EndTickFunction);
+		}
+
+		PrerequisiteRefs.Add({ InPrerequisiteInstance.Handle, 1});
+	}
+}
+
+void FAnimNextModuleInstance::RemovePrerequisite(FAnimNextModuleInstance& InPrerequisiteInstance)
+{
+	check(IsInGameThread());
+
+	UE::AnimNext::FModuleWriteGuard Guard(this);
+
+	int32 FoundIndex = PrerequisiteRefs.IndexOfByPredicate([&InPrerequisiteInstance](const FPrerequisiteReference& InPrerequisiteRef)
+	{
+		return InPrerequisiteRef.Handle == InPrerequisiteInstance.Handle;
+	});
+
+	ensureAlways(FoundIndex != INDEX_NONE);	// Shouldnt really be calling this if we dont have a prerequisite already
+	if(FoundIndex != INDEX_NONE)
+	{
+		FPrerequisiteReference& FoundReference = PrerequisiteRefs[FoundIndex]; 
+		FoundReference.ReferenceCount--;
+		if(FoundReference.ReferenceCount == 0)
+		{
+			// Remove dependency on the prerequisite's end tick function
+			EndTickFunction.RemovePrerequisite(InPrerequisiteInstance.Object, InPrerequisiteInstance.EndTickFunction);
+			for(UE::AnimNext::FModuleEventTickFunction& TickFunction : TickFunctions)
+			{
+				TickFunction.RemovePrerequisite(InPrerequisiteInstance.Object, InPrerequisiteInstance.EndTickFunction);
+			}
+
+			PrerequisiteRefs.RemoveAtSwap(FoundIndex, EAllowShrinking::No);
+		}
+	}
+}
+
+bool FAnimNextModuleInstance::IsPrerequisite(const FAnimNextModuleInstance& InPrerequisiteInstance) const
+{
+	UE::AnimNext::FModuleWriteGuard Guard(this);
+	return PrerequisiteRefs.ContainsByPredicate([&InPrerequisiteInstance](const FPrerequisiteReference& InPrerequisiteRef)
+	{
+		return InPrerequisiteRef.Handle == InPrerequisiteInstance.Handle;
+	});
+}
+
+void FAnimNextModuleInstance::ForEachPrerequisite(TFunctionRef<void(FAnimNextModuleInstance& InPrerequisiteInstance)> InFunction) const
+{
+	UE::AnimNext::FModuleWriteGuard Guard(this);
+
+	if(Pool == nullptr)
+	{
+		return;
+	}
+
+	for(const FPrerequisiteReference& PrerequisiteHandle : PrerequisiteRefs)
+	{
+		FAnimNextModuleInstance* PrerequisiteInstance = Pool->TryGet(PrerequisiteHandle.Handle);
+		if(PrerequisiteInstance == nullptr)
+		{
+			continue;
+		}
+
+		InFunction(*PrerequisiteInstance);
+	}
+}
+
 void FAnimNextModuleInstance::RunRigVMEvent(FName InEventName, float InDeltaTime)
 {
+	UE::AnimNext::FModuleWriteGuard PrerequisiteGuard(this);
+
 	URigVM* VM = GetModule()->RigVM;
 	if(VM == nullptr)
 	{
@@ -565,4 +672,12 @@ void FAnimNextModuleInstance::RunRigVMEvent(FName InEventName, float InDeltaTime
 
 	// Run the VM for this event
 	VM->ExecuteVM(ExtendedExecuteContext, InEventName);
+}
+
+
+TArrayView<UE::AnimNext::FModuleEventTickFunction> FAnimNextModuleInstance::GetTickFunctions()
+{
+	UE::AnimNext::FModuleWriteGuard PrerequisiteGuard(this);
+
+	return TickFunctions;
 }
\ No newline at end of file
diff --git a/Engine/Plugins/Experimental/AnimNext/Source/AnimNext/Private/Module/ModuleGuard.cpp b/Engine/Plugins/Experimental/AnimNext/Source/AnimNext/Private/Module/ModuleGuard.cpp
new file mode 100644
--- /dev/null
+++ b/Engine/Plugins/Experimental/AnimNext/Source/AnimNext/Private/Module/ModuleGuard.cpp
@@ -0,0 +1,66 @@
+﻿// Copyright Epic Games, Inc. All Rights Reserved.
+
+#include "Module/ModuleGuard.h"
+#include "AnimNextPool.h"
+#include "Module/AnimNextModuleInstance.h"
+
+namespace UE::AnimNext
+{
+
+FModuleWriteGuard::FModuleWriteGuard(const FAnimNextModuleInstance* InModuleInstance)
+	: ModuleInstance(InModuleInstance)
+{
+#if ENABLE_MT_DETECTOR
+	if(ModuleInstance)
+	{
+		ModuleInstance->AccessDetector.AcquireWriteAccess();
+
+		if(ModuleInstance->Pool)
+		{
+			// Cache the handles we acquire for later release
+			PrerequisiteHandles.Reserve(ModuleInstance->PrerequisiteRefs.Num());
+			for(const FAnimNextModuleInstance::FPrerequisiteReference& PrerequisiteRef : ModuleInstance->PrerequisiteRefs)
+			{
+				PrerequisiteHandles.Add(PrerequisiteRef.Handle);
+
+				FAnimNextModuleInstance* OtherModuleInstance = ModuleInstance->Pool->TryGet(PrerequisiteRef.Handle);
+				ensureAlways(OtherModuleInstance != nullptr);
+				if(OtherModuleInstance == nullptr)
+				{
+					continue;
+				}
+
+				OtherModuleInstance->AccessDetector.AcquireWriteAccess();
+			}
+		}
+	}
+#endif
+}
+
+FModuleWriteGuard::~FModuleWriteGuard()
+{
+#if ENABLE_MT_DETECTOR
+	if(ModuleInstance)
+	{
+		if(ModuleInstance->Pool)
+		{
+			// Release only the handles we acquired as we may have change prerequisites in this scope
+			for(const FModuleHandle& PrerequisiteHandle : PrerequisiteHandles)
+			{
+				FAnimNextModuleInstance* OtherModuleInstance = ModuleInstance->Pool->TryGet(PrerequisiteHandle);
+				ensureAlways(OtherModuleInstance != nullptr);
+				if(OtherModuleInstance == nullptr)
+				{
+					continue;
+				}
+
+				OtherModuleInstance->AccessDetector.ReleaseWriteAccess();
+			}
+		}
+
+		ModuleInstance->AccessDetector.ReleaseWriteAccess();
+	}
+#endif
+}
+
+}
diff --git a/Engine/Plugins/Experimental/AnimNext/Source/AnimNext/Private/Module/ModuleTaskContext.cpp b/Engine/Plugins/Experimental/AnimNext/Source/AnimNext/Private/Module/ModuleTaskContext.cpp
--- a/Engine/Plugins/Experimental/AnimNext/Source/AnimNext/Private/Module/ModuleTaskContext.cpp
+++ b/Engine/Plugins/Experimental/AnimNext/Source/AnimNext/Private/Module/ModuleTaskContext.cpp
@@ -3,6 +3,7 @@
 #include "Module/ModuleTaskContext.h"
 
 #include "Module/AnimNextModuleInstance.h"
+#include "Module/ModuleGuard.h"
 
 namespace UE::AnimNext
 {
@@ -14,7 +15,9 @@ FModuleTaskContext::FModuleTaskContext(FAnimNextModuleInstance& InModuleInstance
 
 void FModuleTaskContext::QueueInputTraitEvent(FAnimNextTraitEventPtr Event) const
 {
+	FModuleWriteGuard Guard(ModuleInstance);
+
 	ModuleInstance->QueueInputTraitEvent(MoveTemp(Event));
 }
 
-}
\ No newline at end of file
+}
diff --git a/Engine/Plugins/Experimental/AnimNext/Source/AnimNext/Private/Module/ModuleTickFunction.cpp b/Engine/Plugins/Experimental/AnimNext/Source/AnimNext/Private/Module/ModuleTickFunction.cpp
--- a/Engine/Plugins/Experimental/AnimNext/Source/AnimNext/Private/Module/ModuleTickFunction.cpp
+++ b/Engine/Plugins/Experimental/AnimNext/Source/AnimNext/Private/Module/ModuleTickFunction.cpp
@@ -13,6 +13,7 @@
 #include "Module/RigUnit_AnimNextModuleEvents.h"
 #include "Module/ModuleEvents.h"
 #include "TraitCore/TraitEventList.h"
+#include "Module/ModuleGuard.h"
 
 DEFINE_STAT(STAT_AnimNext_Module_EventTick);
 DEFINE_STAT(STAT_AnimNext_Module_EndTick);
@@ -29,6 +30,8 @@ void FModuleEndTickFunction::Run(float DeltaTime)
 {
 	SCOPE_CYCLE_COUNTER(STAT_AnimNext_Module_EndTick);
 
+	FModuleWriteGuard Guard(ModuleInstance);
+
 	// Give module a chance to finish up processing
 	ModuleInstance->EndExecution(DeltaTime);
 
@@ -116,6 +119,8 @@ void FModuleEventTickFunction::Run(float InDeltaTime)
 {
 	SCOPE_CYCLE_COUNTER(STAT_AnimNext_Module_EventTick);
 
+	FModuleWriteGuard Guard(ModuleInstance);
+
 	while (!PreExecuteTasks.IsEmpty())
 	{
 		TOptional<TUniqueFunction<void(const FModuleTaskContext&)>> Function = PreExecuteTasks.Dequeue();
@@ -123,8 +128,14 @@ void FModuleEventTickFunction::Run(float InDeltaTime)
 		Function.GetValue()(FModuleTaskContext(*ModuleInstance));
 	}
 
+	ModuleInstance->RaiseTraitEvents(ModuleInstance->InputEventList);
+
+	OnPreModuleEvent.Broadcast(FModuleTaskContext(*ModuleInstance));
+
 	ModuleInstance->RunRigVMEvent(EventName, InDeltaTime);
 
+	ModuleInstance->RaiseTraitEvents(ModuleInstance->OutputEventList);
+
 	while (!PostExecuteTasks.IsEmpty())
 	{
 		TOptional<TUniqueFunction<void(const FModuleTaskContext&)>> Function = PostExecuteTasks.Dequeue();
diff --git a/Engine/Plugins/Experimental/AnimNext/Source/AnimNext/Private/Module/RigUnit_AnimNextModuleEvents.cpp b/Engine/Plugins/Experimental/AnimNext/Source/AnimNext/Private/Module/RigUnit_AnimNextModuleEvents.cpp
--- a/Engine/Plugins/Experimental/AnimNext/Source/AnimNext/Private/Module/RigUnit_AnimNextModuleEvents.cpp
+++ b/Engine/Plugins/Experimental/AnimNext/Source/AnimNext/Private/Module/RigUnit_AnimNextModuleEvents.cpp
@@ -16,16 +16,43 @@ namespace UE::AnimNext::CVars
 		TEXT("If true, then schedule the AnimNext binding execution in a high priority tick group before other ticks."));
 }
 
-FRigUnit_AnimNextExecuteBindings_Execute()
+FRigUnit_AnimNextExecuteBindings_GT_Execute()
 {
 }
 
-UE::AnimNext::FModuleEventBindingFunction FRigUnit_AnimNextExecuteBindings::GetBindingFunction() const
+UE::AnimNext::FModuleEventBindingFunction FRigUnit_AnimNextExecuteBindings_GT::GetBindingFunction() const
 {
 	return [](const UE::AnimNext::FTickFunctionBindingContext& InContext, FTickFunction& InTickFunction)
 	{
 		InTickFunction.TickGroup = ETickingGroup::TG_PrePhysics;
-		InTickFunction.bRunOnAnyThread = false; // TODO: This will need to vary depending on whether the event is deemed thread-safe 
+		InTickFunction.bRunOnAnyThread = false;
+
+		// Add prerequisite on this for all subsequent tick functions (assumes tick functions are sorted by phase)
+		TArrayView<UE::AnimNext::FModuleEventTickFunction> TickFunctions = InContext.ModuleInstance.GetTickFunctions();
+		for(int32 TickFunctionIndex = InContext.EventIndex + 1; TickFunctionIndex <TickFunctions.Num(); ++TickFunctionIndex)
+		{
+			FTickFunction& TickFunction = TickFunctions[TickFunctionIndex];
+			TickFunction.AddPrerequisite(InContext.ModuleInstance.GetObject(), TickFunctions[InContext.EventIndex]);
+		}
+
+		const bool bHighPriority = UE::AnimNext::CVars::CVarHighPriorityAnimNextExecuteBindingsTick.GetValueOnGameThread();
+		if (InTickFunction.bHighPriority != bHighPriority)
+		{
+			InTickFunction.SetPriorityIncludingPrerequisites(bHighPriority);
+		}
+	};
+}
+
+FRigUnit_AnimNextExecuteBindings_WT_Execute()
+{
+}
+
+UE::AnimNext::FModuleEventBindingFunction FRigUnit_AnimNextExecuteBindings_WT::GetBindingFunction() const
+{
+	return [](const UE::AnimNext::FTickFunctionBindingContext& InContext, FTickFunction& InTickFunction)
+	{
+		InTickFunction.TickGroup = ETickingGroup::TG_PrePhysics;
+		InTickFunction.bRunOnAnyThread = true; 
 
 		// Add prerequisite on this for all subsequent tick functions (assumes tick functions are sorted by phase)
 		TArrayView<UE::AnimNext::FModuleEventTickFunction> TickFunctions = InContext.ModuleInstance.GetTickFunctions();
diff --git a/Engine/Plugins/Experimental/AnimNext/Source/AnimNext/Public/AnimNextRigVMAsset.h b/Engine/Plugins/Experimental/AnimNext/Source/AnimNext/Public/AnimNextRigVMAsset.h
new file mode 100644
--- /dev/null
+++ b/Engine/Plugins/Experimental/AnimNext/Source/AnimNext/Public/AnimNextRigVMAsset.h
@@ -0,0 +1,81 @@
+// Copyright Epic Games, Inc. All Rights Reserved.
+
+#pragma once
+
+#include "CoreMinimal.h"
+#include "RigVMCore/RigVMExecuteContext.h"
+#include "RigVMHost.h"
+#include "AnimNextRigVMAsset.generated.h"
+
+class UAnimNextVariableEntry;
+class UAnimNextRigVMAssetEditorData;
+struct FAnimNextModuleInstance;
+
+namespace UE::AnimNext
+{
+	struct FProxyVariablesContext;
+	struct FModuleEventTickFunction;
+}
+
+namespace UE::AnimNext::UncookedOnly
+{
+	struct FUtils;
+}
+
+/** Base class for all AnimNext assets that can host RigVM logic */
+UCLASS(Abstract)
+class ANIMNEXT_API UAnimNextRigVMAsset : public URigVMHost
+{
+	GENERATED_BODY()
+
+protected:
+	friend class UAnimNextVariableEntry;
+	friend class UAnimNextRigVMAssetEditorData;
+	friend struct FAnimNextModuleInstance;
+	friend struct UE::AnimNext::FModuleEventTickFunction;
+	friend struct UE::AnimNext::FProxyVariablesContext;
+	friend struct UE::AnimNext::UncookedOnly::FUtils;
+
+	UAnimNextRigVMAsset(const FObjectInitializer& ObjectInitializer);
+
+	// UObject interface
+	virtual void BeginDestroy() override;
+	virtual void PostLoad() override;
+	virtual void GetAssetRegistryTags(FAssetRegistryTagsContext Context) const override;
+	virtual void PreDuplicate(FObjectDuplicationParameters& DupParams) override;
+
+public:
+	// Get public variable defaults property bag
+	const FInstancedPropertyBag& GetPublicVariableDefaults() const { return PublicVariableDefaults; }
+
+	// Get the RigVM object we host
+	const URigVM* GetRigVM() const { return RigVM; }
+
+private:
+	// URigVMHost interface
+	virtual TArray<FRigVMExternalVariable> GetExternalVariablesImpl(bool bFallbackToBlueprint) const;
+
+protected:
+	// The ExtendedExecuteContext object holds the common work data used by the RigVM internals. It is populated during the initial VM initialization.
+	// Each instance of an AnimGraph requires a copy of this context and a call to initialize the VM instance with the context copy, 
+	// so the cached memory handles are updated to the correct memory addresses.
+	// This context is used as a reference to copy the common data for all instances created.
+	UPROPERTY(Transient)
+	FRigVMExtendedExecuteContext ExtendedExecuteContext;
+
+	UPROPERTY()
+	TObjectPtr<URigVM> RigVM;
+
+	// Variables and their defaults (including public variables, sorted first)
+	UPROPERTY()
+	FInstancedPropertyBag VariableDefaults;
+
+	// Public variables (for easy duplication)
+	UPROPERTY()
+	FInstancedPropertyBag PublicVariableDefaults;
+
+#if WITH_EDITORONLY_DATA
+	UPROPERTY(VisibleAnywhere, Instanced, Category = "Editor Data", meta = (ShowInnerProperties))
+	TObjectPtr<UObject> EditorData;
+#endif
+};
diff --git a/Engine/Plugins/Experimental/AnimNext/Source/AnimNext/Public/Component/AnimNextComponent.h b/Engine/Plugins/Experimental/AnimNext/Source/AnimNext/Public/Component/AnimNextComponent.h
new file mode 100644
--- /dev/null
+++ b/Engine/Plugins/Experimental/AnimNext/Source/AnimNext/Public/Component/AnimNextComponent.h
@@ -0,0 +1,233 @@
+// Copyright Epic Games, Inc. All Rights Reserved.
+
+#pragma once
+
+#include "AnimNextPublicVariablesProxy.h"
+#include "Components/ActorComponent.h"
+#include "Module/ModuleHandle.h"
+#include "Param/ParamType.h"
+#include "TraitCore/TraitEvent.h"
+#include "Variables/IAnimNextVariableProxyHost.h"
+#include "Module/AnimNextModuleInitMethod.h"
+
+#include "AnimNextComponent.generated.h"
+
+struct FAnimNextComponentInstanceData;
+class UAnimNextComponentWorldSubsystem;
+class UAnimNextModule;
+
+namespace UE::AnimNext
+{
+	struct FProxyVariablesContext;
+}
+
+namespace UE::AnimNext::UncookedOnly
+{
+	struct FUtils;
+}
+
+UCLASS(MinimalAPI, meta = (BlueprintSpawnableComponent))
+class UAnimNextComponent : public UActorComponent, public IAnimNextVariableProxyHost
+{
+	GENERATED_BODY()
+
+	// UActorComponent interface
+	virtual void OnRegister() override;
+	virtual void OnUnregister() override;
+	virtual void BeginPlay() override;
+	virtual void EndPlay(const EEndPlayReason::Type EndPlayReason) override;
+
+	// IAnimNextVariableProxyHost interface
+	virtual void FlipPublicVariablesProxy(const UE::AnimNext::FProxyVariablesContext& InContext) override;
+
+#if WITH_EDITOR
+	// Called back to refresh any cached data on module compilation
+	void OnModuleCompiled();
+#endif
+
+	// (Re-)create the public variable proxy
+	void CreatePublicVariablesProxy();
+
+	// (Re-)create the public variable proxy
+	void DestroyPublicVariablesProxy();
+
+public:
+	// Sets a module variable's value.
+	// @param    Name     The name of the variable to set
+	// @param    Value    The value to set the variable to
+	UFUNCTION(BlueprintCallable, Category = "AnimNext", CustomThunk, meta = (CustomStructureParam = Value, UnsafeDuringActorConstruction))
+	ANIMNEXT_API void SetVariable(UPARAM(meta = (CustomWidget = "VariableName")) FName Name, int32 Value);
+
+	/**
+	 * Module variable value getters.
+	 * Numeric types (bool, (u)int32, (u)int64, float, double) support type conversion.
+	 * Struct & Object types will be const, if you need to modify them use SetVariable[Struct, Object]Ref
+	 * @param    Name     The name of the variable to get
+	 * @return            The variable value on success, None otherwise
+	 */
+	ANIMNEXT_API TOptional<bool> GetVariableBool(const FName Name) const;
+	ANIMNEXT_API TOptional<uint8> GetVariableByte(const FName Name) const;
+	ANIMNEXT_API TOptional<int32> GetVariableInt32(const FName Name) const;
+	ANIMNEXT_API TOptional<uint32> GetVariableUInt32(const FName Name) const;
+	ANIMNEXT_API TOptional<int64> GetVariableInt64(const FName Name) const;
+	ANIMNEXT_API TOptional<uint64> GetVariableUInt64(const FName Name) const; 
+	ANIMNEXT_API TOptional<float> GetVariableFloat(const FName Name) const;
+	ANIMNEXT_API TOptional<double> GetVariableDouble(const FName Name) const;
+	ANIMNEXT_API TOptional<FName> GetVariableName(const FName Name) const;
+	ANIMNEXT_API TOptional<FString> GetVariableString(const FName Name) const;
+	ANIMNEXT_API TOptional<uint8> GetVariableEnum(const FName Name, const UEnum* RequestedEnum) const;
+	ANIMNEXT_API TOptional<FConstStructView> GetVariableStruct(const FName Name, const UScriptStruct* RequestedStruct = nullptr) const;
+	ANIMNEXT_API TOptional<const UObject*> GetVariableObject(const FName Name, const UClass* RequestedClass = nullptr) const;
+	ANIMNEXT_API TOptional<const UClass*> GetVariableClass(const FName Name) const;
+	ANIMNEXT_API TOptional<FSoftObjectPath> GetVariableSoftPath(const FName Name) const;
+
+	/** @return enum value of specified type. */
+	template <typename T>
+	TOptional<T> GetVariableEnum(const FName Name) const
+	{
+		static_assert(TIsEnum<T>::Value, "Should only call this with enum types");
+
+		TOptional<uint8> Result = GetVariableEnum(Name, StaticEnum<T>());
+		if (!Result.IsSet())
+		{
+			return TOptional<T>();
+		}
+		return TOptional<T>(static_cast<T>(Result.GetValue()));
+	}
+
+	/** @return struct reference of specified type. */
+	template <typename T>
+	TOptional<const T*> GetVariableStruct(const FName Name) const
+	{
+		TOptional<FConstStructView> Result = GetVariableStruct(Name, TBaseStructure<T>::Get());
+		if (!Result.IsSet())
+		{
+			return TOptional<const T*>();
+		}
+		if (const T* ValuePtr = Result.GetValue().GetPtr<const T>())
+		{
+			return TOptional<const T*>(ValuePtr);
+		}
+		return TOptional<const T*>();
+	}
+
+	/** @return object pointer value of specified type. */
+	template <typename T>
+	TOptional<const T*> GetVariableObject(const FName Name) const
+	{
+		static_assert(TIsDerivedFrom<T, UObject>::Value, "Should only call this with object types");
+
+		TOptional<const UObject*> Result = GetVariableObject(Name, T::StaticClass());
+		if (!Result.IsSet())
+		{
+			return TOptional<const T*>();
+		}
+		if (Result.GetValue() == nullptr)
+		{
+			return TOptional<const T*>(nullptr);
+		}
+		if (const T* Object = Cast<const T>(Result.GetValue()))
+		{
+			return TOptional<const T*>(Object);
+		}
+		return TOptional<const T*>();
+	}
+
+	/**
+	 * Module variable value setters.
+	 * @param    Name     The name of the variable to set
+	 * @param    Value    The value to set the variable to
+	 * @return            True on success, false otherwise
+	 */
+	ANIMNEXT_API bool SetVariableBool(const FName Name, const bool bInValue);
+	ANIMNEXT_API bool SetVariableByte(const FName Name, const uint8 InValue);
+	ANIMNEXT_API bool SetVariableInt32(const FName Name, const int32 InValue);
+	ANIMNEXT_API bool SetVariableUInt32(const FName Name, const uint32 InValue);
+	ANIMNEXT_API bool SetVariableInt64(const FName Name, const int64 InValue);
+	ANIMNEXT_API bool SetVariableUInt64(const FName Name, const uint64 InValue);
+	ANIMNEXT_API bool SetVariableFloat(const FName Name, const float InValue);
+	ANIMNEXT_API bool SetVariableDouble(const FName Name, const double InValue);
+	ANIMNEXT_API bool SetVariableName(const FName Name, const FName InValue);
+	ANIMNEXT_API bool SetVariableString(const FName Name, const FString& InValue);
+	ANIMNEXT_API bool SetVariableEnum(const FName Name, const uint8 InValue, const UEnum* Enum);
+	ANIMNEXT_API bool SetVariableStruct(const FName Name, FConstStructView InValue);
+	ANIMNEXT_API bool SetVariableStructRef(const FName Name, TFunctionRef<void(FStructView)> InStructRefSetter, const UScriptStruct* RequestedStruct = nullptr);
+	ANIMNEXT_API bool SetVariableObject(const FName Name, UObject* InValue);
+	ANIMNEXT_API bool SetVariableObjectRef(const FName Name, TFunctionRef<void(UObject*)> InObjectRefSetter, const UClass* RequestedClass = nullptr);
+	ANIMNEXT_API bool SetVariableClass(const FName Name, UClass* InValue);
+	ANIMNEXT_API bool SetVariableSoftPath(const FName Name, const FSoftObjectPath& InValue);
+	ANIMNEXT_API bool SetVariableArrayRef(const FName Name, TFunctionRef<void(FPropertyBagArrayRef&)> InArrayRefSetter);
+
+	/** Sets enum variable with specified type. */
+	template <typename T>
+	bool SetVariableEnum(const FName Name, const T InValue)
+	{
+		static_assert(TIsEnum<T>::Value, "Should only call this with enum types");
+		return SetVariableEnum(Name, static_cast<uint8>(InValue), StaticEnum<T>());
+	}
+
+	/** Sets struct variable with specified type. */
+	template <typename T>
+	bool SetVariableStruct(const FName Name, const T& InValue)
+	{
+		return SetVariableStruct(Name, FConstStructView::Make(InValue));
+	}
+
+	/** Sets object pointer variable with specified type. */
+	template <typename T>
+	bool SetVariableObject(const FName Name, T* InValue)
+	{
+		static_assert(TIsDerivedFrom<T, UObject>::Value, "Should only call this with object types");
+		return SetVariableObject(Name, (UObject*)InValue);
+	}
+
+	// Enable or disable this component's update
+	UFUNCTION(BlueprintCallable, Category = "AnimNext")
+	ANIMNEXT_API void SetEnabled(bool bEnabled);
+
+	// Enable or disable debug drawing. Note only works in builds with UE_ENABLE_DEBUG_DRAWING enabled
+	UFUNCTION(BlueprintCallable, Category = "AnimNext")
+	ANIMNEXT_API void ShowDebugDrawing(bool bShowDebugDrawing);
+
+	// Queues an input trait event
+	// Input events will be processed in the next graph update after they are queued
+	ANIMNEXT_API void QueueInputTraitEvent(FAnimNextTraitEventPtr Event);
+
+	// Get the handle to the registered module
+	ANIMNEXT_API UE::AnimNext::FModuleHandle GetModuleHandle() const { return ModuleHandle; }
+
+private:
+	DECLARE_FUNCTION(execSetVariable);
+
+private:
+	friend struct UE::AnimNext::UncookedOnly::FUtils;
+	friend class UAnimNextComponentWorldSubsystem;
+
+	// The AnimNext module that this component will run
+	UPROPERTY(EditAnywhere, Category="Module")
+	TObjectPtr<UAnimNextModule> Module = nullptr;
+
+	UPROPERTY(Transient)
+	TObjectPtr<UAnimNextComponentWorldSubsystem> Subsystem = nullptr;
+
+	// Handle to the registered module
+	UE::AnimNext::FModuleHandle ModuleHandle;
+
+	// Lock for public variables proxy
+	mutable FRWLock PublicVariablesLock;
+
+	// Proxy public variables
+	UPROPERTY()
+	FAnimNextPublicVariablesProxy PublicVariablesProxy;
+
+	// Map from name->proxy variable index
+	TMap<FName, int32> PublicVariablesProxyMap;
+
+	// How to initialize the module
+	UPROPERTY(EditAnywhere, Category="Module")
+	EAnimNextModuleInitMethod InitMethod = EAnimNextModuleInitMethod::InitializeAndPauseInEditor;
+
+	/** When checked, the module's debug drawing instructions are drawn in the viewport */
+	UPROPERTY(EditAnywhere, Category = "Rendering")
+	uint8 bShowDebugDrawing : 1 = false;
+};
diff --git a/Engine/Plugins/Experimental/AnimNext/Source/AnimNext/Public/Component/AnimNextPublicVariablesProxy.h b/Engine/Plugins/Experimental/AnimNext/Source/AnimNext/Public/Component/AnimNextPublicVariablesProxy.h
new file mode 100644
--- /dev/null
+++ b/Engine/Plugins/Experimental/AnimNext/Source/AnimNext/Public/Component/AnimNextPublicVariablesProxy.h
@@ -0,0 +1,48 @@
+﻿// Copyright Epic Games, Inc. All Rights Reserved.
+
+#pragma once
+
+#include "StructUtils/PropertyBag.h"
+#include "AnimNextPublicVariablesProxy.generated.h"
+
+class UAnimNextComponent;
+struct FAnimNextModuleInstance;
+
+// Proxy struct for public variables held on the component
+// TODO: Instead of a dirty flag bitfield (and copying parameters), with some minor refactoring of RigVM external variables
+// we could switch this to use a double-buffered external variables array too, so each write to a variable from the game thread
+// would just write to the current 'GT-side' buffer, then update the external variable ptr to point to the 'latest' version written
+// Then when it comes to swap buffers before WT execution, we just swap the ptrs for the external variables and dont have to copy. 
+USTRUCT()
+struct FAnimNextPublicVariablesProxy
+{
+	GENERATED_BODY()
+
+private:
+	friend UAnimNextComponent;
+	friend FAnimNextModuleInstance;
+
+	void Reset()
+	{
+		Data.Reset();
+		DirtyFlags.Reset();
+		bIsDirty = false;
+	}
+
+	void Empty()
+	{
+		Data.Reset();
+		DirtyFlags.Empty();
+		bIsDirty = false;
+	}
+
+	// Proxy public variables
+	UPROPERTY()
+	FInstancedPropertyBag Data;
+
+	// Dirty flags for each public variable
+	TBitArray<> DirtyFlags;
+
+	// Global dirty flag
+	bool bIsDirty = false;
+};
\ No newline at end of file
diff --git a/Engine/Plugins/Experimental/AnimNext/Source/AnimNext/Public/DataInterface/AnimNextDataInterface.h b/Engine/Plugins/Experimental/AnimNext/Source/AnimNext/Public/DataInterface/AnimNextDataInterface.h
new file mode 100644
--- /dev/null
+++ b/Engine/Plugins/Experimental/AnimNext/Source/AnimNext/Public/DataInterface/AnimNextDataInterface.h
@@ -0,0 +1,82 @@
+﻿// Copyright Epic Games, Inc. All Rights Reserved.
+
+#pragma once
+
+#include "CoreMinimal.h"
+#include "AnimNextRigVMAsset.h"
+#include "AnimNextDataInterface.generated.h"
+
+class UAnimNextDataInterfaceFactory;
+class UAnimNextDataInterface;
+
+namespace UE::AnimNext
+{
+	struct FInjectionInfo;
+}
+
+namespace UE::AnimNext::Tests
+{
+	class FDataInterfaceCompile;
+}
+
+namespace UE::AnimNext::UncookedOnly
+{
+	struct FUtils;
+}
+
+// Information about an implemented interface
+USTRUCT()
+struct FAnimNextImplementedDataInterface
+{
+	GENERATED_BODY()
+
+	// Reference to the implemented data interface asset
+	UPROPERTY()
+	TObjectPtr<const UAnimNextDataInterface> DataInterface;
+
+	// Native struct used for communication with the data interface
+	UPROPERTY()
+	TObjectPtr<const UScriptStruct> NativeInterface;
+
+	// Index of the first variable that implements the interface
+	UPROPERTY()
+	int32 VariableIndex = INDEX_NONE;
+
+	// Number of variables that implement the interface
+	UPROPERTY()
+	int32 NumVariables = 0;
+
+	// Whether to automatically bind this interface to any host data interface 
+	UPROPERTY()
+	bool bAutoBindToHost = false;
+};
+
+// Data interfaces provide a set of named data that is shared between AnimNext assets and used for communication between assets and functional units
+UCLASS(BlueprintType)
+class ANIMNEXT_API UAnimNextDataInterface : public UAnimNextRigVMAsset
+{
+	GENERATED_BODY()
+
+public:
+	UAnimNextDataInterface(const FObjectInitializer& ObjectInitializer);
+
+	// Get all the implemented interfaces
+	TConstArrayView<FAnimNextImplementedDataInterface> GetImplementedInterfaces() const;
+
+	// Find an implemented interface
+	const FAnimNextImplementedDataInterface* FindImplementedInterface(const UAnimNextDataInterface* InDataInterface) const;
+
+private:
+	friend class UAnimNextDataInterfaceFactory;
+	friend class UE::AnimNext::Tests::FDataInterfaceCompile;
+	friend struct UE::AnimNext::UncookedOnly::FUtils;
+	friend UE::AnimNext::FInjectionInfo;
+
+	// Information about implemented interfaces. Note this includes the 'self' interface (first), if any public variables are specified.
+	UPROPERTY()
+	TArray<FAnimNextImplementedDataInterface> ImplementedInterfaces;
+
+	// The variable index of the default injection site
+	UPROPERTY()
+	int32 DefaultInjectionSiteIndex;
+};
diff --git a/Engine/Plugins/Experimental/AnimNext/Source/AnimNext/Public/DataInterface/AnimNextDataInterfaceInstance.h b/Engine/Plugins/Experimental/AnimNext/Source/AnimNext/Public/DataInterface/AnimNextDataInterfaceInstance.h
--- a/Engine/Plugins/Experimental/AnimNext/Source/AnimNext/Public/DataInterface/AnimNextDataInterfaceInstance.h
+++ b/Engine/Plugins/Experimental/AnimNext/Source/AnimNext/Public/DataInterface/AnimNextDataInterfaceInstance.h
@@ -10,6 +10,12 @@
 #include "AnimNextDataInterfaceInstance.generated.h"
 
 class UAnimNextDataInterface;
+struct FAnimNextModuleInjectionComponent;
+
+namespace UE::AnimNext
+{
+	struct FInjectionInfo;
+}
 
 // Base struct for data interface-derived instances
 USTRUCT()
@@ -58,6 +64,9 @@ struct ANIMNEXT_API FAnimNextDataInterfaceInstance
 		return GetVariableInternal(InVariableName, FAnimNextParamType::GetType<ValueType>(), TArrayView<uint8>(reinterpret_cast<uint8*>(&OutResult), sizeof(ValueType)));
 	}
 
+	// Get the instance (graph, module etc.) that owns/hosts us
+	FAnimNextDataInterfaceInstance* GetHost() const { return HostInstance; }
+
 private:
 	// Helper function for GetVariable
 	EPropertyBagResult GetVariableInternal(FName InVariableName, const FAnimNextParamType& InType, TArrayView<uint8> OutResult) const;
@@ -74,4 +83,10 @@ protected:
 	// Extended execute context instance for this graph instance, we own it
 	UPROPERTY(Transient)
 	FRigVMExtendedExecuteContext ExtendedExecuteContext;
+
+	// The instance (graph, module etc.) that owns/hosts us
+	FAnimNextDataInterfaceInstance* HostInstance = nullptr;
+
+	friend UE::AnimNext::FInjectionInfo;
+	friend FAnimNextModuleInjectionComponent;
 };
diff --git a/Engine/Plugins/Experimental/AnimNext/Source/AnimNext/Public/Variables/IAnimNextVariableProxyHost.h b/Engine/Plugins/Experimental/AnimNext/Source/AnimNext/Public/Variables/IAnimNextVariableProxyHost.h
new file mode 100644
--- /dev/null
+++ b/Engine/Plugins/Experimental/AnimNext/Source/AnimNext/Public/Variables/IAnimNextVariableProxyHost.h
@@ -0,0 +1,31 @@
+﻿// Copyright Epic Games, Inc. All Rights Reserved.
+
+#pragma once
+
+#include "UObject/Interface.h"
+#include "IAnimNextVariableProxyHost.generated.h"
+
+struct FAnimNextModuleInstance;
+
+namespace UE::AnimNext
+{
+	struct FProxyVariablesContext;
+}
+
+/** Interface used for hosting proxy variables */
+UINTERFACE()
+class UAnimNextVariableProxyHost : public UInterface
+{
+	GENERATED_BODY()
+};
+
+class IAnimNextVariableProxyHost
+{
+	GENERATED_BODY()
+
+	// Flip the public variables proxy buffer
+	// Called from any thread to retrieve the latest copy of any public proxy variables
+	virtual void FlipPublicVariablesProxy(const UE::AnimNext::FProxyVariablesContext& InContext) = 0;
+
+	friend struct FAnimNextModuleInstance;
+};
diff --git a/Engine/Plugins/Experimental/AnimNext/Source/AnimNextTestSuite/Private/AnimNextVariablesTest.cpp b/Engine/Plugins/Experimental/AnimNext/Source/AnimNextTestSuite/Private/AnimNextVariablesTest.cpp
--- a/Engine/Plugins/Experimental/AnimNext/Source/AnimNextTestSuite/Private/AnimNextVariablesTest.cpp
+++ b/Engine/Plugins/Experimental/AnimNext/Source/AnimNextTestSuite/Private/AnimNextVariablesTest.cpp
@@ -405,12 +405,12 @@ bool FVariables_UOLBindings::RunTest(const FString& InParameters)
 	Asset->GetRigVMExtendedExecuteContext().SetRuntimeSettings(RuntimeSettings);
 	
 	FAnimNextExecuteContext& AnimNextContext = Asset->GetRigVMExtendedExecuteContext().GetPublicDataSafe<FAnimNextExecuteContext>();
-	FAnimNextModuleInstance Instance(Asset, NewObject<UAnimNextComponent>(), EAnimNextModuleInitMethod::None);
+	FAnimNextModuleInstance Instance(Asset, NewObject<UAnimNextComponent>(), nullptr, EAnimNextModuleInitMethod::None);
 	FAnimNextModuleContextData ContextData(&Instance);
 	FScopedExecuteContextData ContextDataScope(AnimNextContext, ContextData);
 
 	// Execute bindings to fetch data
-	Asset->GetVM()->ExecuteVM(Asset->GetRigVMExtendedExecuteContext(), FRigUnit_AnimNextExecuteBindings::EventName);
+	Asset->GetVM()->ExecuteVM(Asset->GetRigVMExtendedExecuteContext(), FRigUnit_AnimNextExecuteBindings_GT::EventName);
 
 	// Run PrePhysics event to print messages
 	Asset->GetVM()->ExecuteVM(Asset->GetRigVMExtendedExecuteContext(), FRigUnit_AnimNextPrePhysicsEvent::EventName);
diff --git a/Engine/Plugins/Experimental/AnimNext/Source/AnimNextUncookedOnly/Internal/DataInterface/AnimNextDataInterface_EditorData.h b/Engine/Plugins/Experimental/AnimNext/Source/AnimNextUncookedOnly/Internal/DataInterface/AnimNextDataInterface_EditorData.h
--- a/Engine/Plugins/Experimental/AnimNext/Source/AnimNextUncookedOnly/Internal/DataInterface/AnimNextDataInterface_EditorData.h
+++ b/Engine/Plugins/Experimental/AnimNext/Source/AnimNextUncookedOnly/Internal/DataInterface/AnimNextDataInterface_EditorData.h
@@ -9,6 +9,11 @@
 
 class UAnimNextDataInterfaceFactory;
 
+namespace UE::AnimNext::UncookedOnly
+{
+	struct FUtils;
+}
+
 /** Editor data for AnimNext data interfaces */
 UCLASS()
 class ANIMNEXTUNCOOKEDONLY_API UAnimNextDataInterface_EditorData : public UAnimNextRigVMAssetEditorData
@@ -22,4 +27,9 @@ protected:
 	virtual void CustomizeNewAssetEntry(UAnimNextRigVMAssetEntry* InNewEntry) const override;
 
 	friend class UAnimNextDataInterfaceFactory;
+	friend UE::AnimNext::UncookedOnly::FUtils;
+
+	// The injection site (AnimNext Graph) that will be used by default
+	UPROPERTY(EditAnywhere, Category=Injection)
+	FName DefaultInjectionSite;
 };
diff --git a/Engine/Plugins/Experimental/AnimNext/Source/AnimNextUncookedOnly/Internal/UncookedOnlyUtils.h b/Engine/Plugins/Experimental/AnimNext/Source/AnimNextUncookedOnly/Internal/UncookedOnlyUtils.h
--- a/Engine/Plugins/Experimental/AnimNext/Source/AnimNextUncookedOnly/Internal/UncookedOnlyUtils.h
+++ b/Engine/Plugins/Experimental/AnimNext/Source/AnimNextUncookedOnly/Internal/UncookedOnlyUtils.h
@@ -241,6 +241,9 @@ struct ANIMNEXTUNCOOKEDONLY_API FUtils
 
 	// Make an event name that we use as a wrapper to call RigVM functions 
 	static FString MakeFunctionWrapperEventName(FName InFunctionName);
+
+private:
+	static void CompileVariableBindingsInternal(const FRigVMCompileSettings& InSettings, UAnimNextRigVMAsset* InAsset, TArray<URigVMGraph*>& OutGraphs, bool bInThreadSafe);
 };
 
 }
\ No newline at end of file
diff --git a/Engine/Plugins/Experimental/AnimNext/Source/AnimNextUncookedOnly/Internal/Variables/AnimNextUniversalObjectLocatorBindingData.h b/Engine/Plugins/Experimental/AnimNext/Source/AnimNextUncookedOnly/Internal/Variables/AnimNextUniversalObjectLocatorBindingData.h
--- a/Engine/Plugins/Experimental/AnimNext/Source/AnimNextUncookedOnly/Internal/Variables/AnimNextUniversalObjectLocatorBindingData.h
+++ b/Engine/Plugins/Experimental/AnimNext/Source/AnimNextUncookedOnly/Internal/Variables/AnimNextUniversalObjectLocatorBindingData.h
@@ -46,4 +46,9 @@ struct FAnimNextUniversalObjectLocatorBindingData : public FAnimNextVariableBind
 		const bool bValidFunction = ((Type == FAnimNextUniversalObjectLocatorBindingType::Function || Type == FAnimNextUniversalObjectLocatorBindingType::HoistedFunction) && !Function.IsNull());
 		return bValidLocator && (bValidProperty || bValidFunction);
 	}
+
+#if WITH_EDITORONLY_DATA
+	// Check if this binding data is safe to execute on a worker thread
+	virtual bool IsThreadSafe() const override;
+#endif
 };
diff --git a/Engine/Plugins/Experimental/AnimNext/Source/AnimNextUncookedOnly/Private/UncookedOnlyUtils.cpp b/Engine/Plugins/Experimental/AnimNext/Source/AnimNextUncookedOnly/Private/UncookedOnlyUtils.cpp
--- a/Engine/Plugins/Experimental/AnimNext/Source/AnimNextUncookedOnly/Private/UncookedOnlyUtils.cpp
+++ b/Engine/Plugins/Experimental/AnimNext/Source/AnimNextUncookedOnly/Private/UncookedOnlyUtils.cpp
@@ -201,7 +201,7 @@ void FUtils::CompileVariables(UAnimNextRigVMAsset* InAsset)
 		}
 	};
 
-	AddDataInterface(DataInterface, nullptr, false, false, AddDataInterface);
+	AddDataInterface(DataInterface, nullptr, false, true, AddDataInterface);
 
 	// Sort public entries first, then by data interface & then by size, largest first, for better packing
 	static_assert(EAnimNextExportAccessSpecifier::Private < EAnimNextExportAccessSpecifier::Public, "Private must be less than Public as parameters are sorted internally according to this assumption");
@@ -225,6 +225,8 @@ void FUtils::CompileVariables(UAnimNextRigVMAsset* InAsset)
 		}
 	});
 
+	DataInterface->DefaultInjectionSiteIndex = INDEX_NONE;
+
 	if(StructEntryInfos.Num() > 0)
 	{
 		// Build PropertyDescs and values to batch-create the property bag
@@ -240,6 +242,13 @@ void FUtils::CompileVariables(UAnimNextRigVMAsset* InAsset)
 			PropertyDescs.Emplace(StructEntryInfo->Name, StructEntryInfo->Type.ContainerType, StructEntryInfo->Type.ValueType, StructEntryInfo->Type.ValueTypeObject);
 			Values.Add(StructEntryInfo->Value);
 
+			// Note: string comparison here because otherwise we would have a circular dependency
+			if( EditorData->DefaultInjectionSite == StructEntryInfo->Name &&
+				StructEntryInfo->Type.ValueTypeObject->GetPathName() == TEXT("/Script/AnimNextAnimGraph.AnimNextAnimGraph"))
+			{
+				DataInterface->DefaultInjectionSiteIndex = StructEntryInfo.GetIndex();
+			}
+			
 			if(StructEntryInfo->AccessSpecifier != EAnimNextExportAccessSpecifier::Public)
 			{
 				continue;
@@ -267,6 +276,11 @@ void FUtils::CompileVariables(UAnimNextRigVMAsset* InAsset)
 			}
 		}
 
+		if(EditorData->DefaultInjectionSite != NAME_None && DataInterface->DefaultInjectionSiteIndex == INDEX_NONE)
+		{
+			Log.Error(FText::Format(LOCTEXT("MissingDefaultInjectionSiteWarning", "Could not find default injection site: {0}"), FText::FromName(EditorData->DefaultInjectionSite)));
+		}
+
 		// Create new property bags and migrate
 		EPropertyBagResult Result = DataInterface->VariableDefaults.ReplaceAllPropertiesAndValues(PropertyDescs, Values);
 		check(Result == EPropertyBagResult::Success);
@@ -292,10 +306,19 @@ void FUtils::CompileVariables(UAnimNextRigVMAsset* InAsset)
 		DataInterface->VariableDefaults.Reset();
 		DataInterface->PublicVariableDefaults.Reset();
 		DataInterface->VM->ClearExternalVariables(DataInterface->ExtendedExecuteContext);
+		DataInterface->DefaultInjectionSiteIndex = INDEX_NONE;
 	}
+
+
 }
 
 void FUtils::CompileVariableBindings(const FRigVMCompileSettings& InSettings, UAnimNextRigVMAsset* InAsset, TArray<URigVMGraph*>& OutGraphs)
+{
+	CompileVariableBindingsInternal(InSettings, InAsset, OutGraphs, true);
+	CompileVariableBindingsInternal(InSettings, InAsset,  OutGraphs, false);
+}
+
+void FUtils::CompileVariableBindingsInternal(const FRigVMCompileSettings& InSettings, UAnimNextRigVMAsset* InAsset, TArray<URigVMGraph*>& OutGraphs, bool bInThreadSafe)
 {
 	check(InAsset);
 
@@ -317,6 +340,11 @@ void FUtils::CompileVariableBindings(const FRigVMCompileSettings& InSettings, UA
 			continue;
 		}
 
+		if(Binding.Get<FAnimNextVariableBindingData>().IsThreadSafe() != bInThreadSafe)
+		{
+			continue;
+		}
+
 		TSharedPtr<IVariableBindingType> BindingType = Module.FindVariableBindingType(Binding.GetScriptStruct());
 		if(!BindingType.IsValid())
 		{
@@ -329,7 +357,7 @@ void FUtils::CompileVariableBindings(const FRigVMCompileSettings& InSettings, UA
 	}
 
 	const bool bHasBindings = BindingGroups.Num() > 0;
-	const bool bHasPublicVariablesToCopy = EditorData->IsA<UAnimNextModule_EditorData>() && EditorData->HasPublicVariables();
+	const bool bHasPublicVariablesToCopy = EditorData->IsA<UAnimNextModule_EditorData>() && EditorData->HasPublicVariables() && bInThreadSafe;
 	if(!bHasBindings && !bHasPublicVariablesToCopy)
 	{
 		// Nothing to do here
@@ -340,7 +368,8 @@ void FUtils::CompileVariableBindings(const FRigVMCompileSettings& InSettings, UA
 
 	FRigVMClient* VMClient = EditorData->GetRigVMClient();
 	URigVMController* Controller = VMClient->GetOrCreateController(BindingGraph);
-	URigVMNode* ExecuteBindingsNode = Controller->AddUnitNode(FRigUnit_AnimNextExecuteBindings::StaticStruct(), FRigVMStruct::ExecuteName, FVector2D::ZeroVector, FString(), false);
+	UScriptStruct* BindingsNodeType = bInThreadSafe ? FRigUnit_AnimNextExecuteBindings_WT::StaticStruct() : FRigUnit_AnimNextExecuteBindings_GT::StaticStruct();
+	URigVMNode* ExecuteBindingsNode = Controller->AddUnitNode(BindingsNodeType, FRigVMStruct::ExecuteName, FVector2D::ZeroVector, FString(), false);
 	if(ExecuteBindingsNode == nullptr)
 	{
 		InSettings.ReportError(TEXT("Could not spawn Execute Bindings node"));
@@ -354,7 +383,8 @@ void FUtils::CompileVariableBindings(const FRigVMCompileSettings& InSettings, UA
 	}
 	URigVMPin* ExecPin = ExecuteBindingsExecPin;
 
-	if(bHasPublicVariablesToCopy)
+	// Copy public vars in the WT event
+	if(bHasPublicVariablesToCopy && bInThreadSafe)
 	{
 		URigVMNode* CopyProxyVariablesNode = Controller->AddUnitNode(FRigUnit_CopyModuleProxyVariables::StaticStruct(), FRigVMStruct::ExecuteName, FVector2D(200, 0.0f), FString(), false);
 		if(CopyProxyVariablesNode == nullptr)
@@ -365,23 +395,24 @@ void FUtils::CompileVariableBindings(const FRigVMCompileSettings& InSettings, UA
 		URigVMPin* CopyProxyVariablesExecPin = CopyProxyVariablesNode->FindPin(FRigVMStruct::ExecuteContextName.ToString());
 		if(ExecPin == nullptr)
 		{
-			InSettings.ReportError(TEXT("Could not find execute pin on Execute Bindings node"));
+			InSettings.ReportError(TEXT("Could not find execute pin on Copy Module Proxy Variables node"));
 			return;
 		}
 		bool bLinkAdded = Controller->AddLink(ExecuteBindingsExecPin, CopyProxyVariablesExecPin, false);
 		if(!bLinkAdded)
 		{
-			InSettings.ReportError(TEXT("Could not link Execute Bindings node"));
+			InSettings.ReportError(TEXT("Could not link Copy Module Proxy Variables node"));
 			return;
 		}
 		ExecPin = CopyProxyVariablesExecPin;
 	}
 
 	IVariableBindingType::FBindingGraphFragmentArgs Args;
-	Args.Event = FRigUnit_AnimNextExecuteBindings::StaticStruct();
+	Args.Event = BindingsNodeType;
 	Args.Controller = Controller;
 	Args.BindingGraph = BindingGraph;
 	Args.ExecTail = ExecPin;
+	Args.bThreadSafe = bInThreadSafe;
 
 	FVector2D Location(0.0f, 0.0f);
 	for(const TPair<IVariableBindingType*, TArray<IVariableBindingType::FBindingGraphInput>>& BindingGroupPair : BindingGroups)
diff --git a/Engine/Plugins/Experimental/AnimNext/Source/AnimNextUncookedOnly/Private/Variables/UniversalObjectLocatorBindingType.cpp b/Engine/Plugins/Experimental/AnimNext/Source/AnimNextUncookedOnly/Private/Variables/UniversalObjectLocatorBindingType.cpp
--- a/Engine/Plugins/Experimental/AnimNext/Source/AnimNextUncookedOnly/Private/Variables/UniversalObjectLocatorBindingType.cpp
+++ b/Engine/Plugins/Experimental/AnimNext/Source/AnimNextUncookedOnly/Private/Variables/UniversalObjectLocatorBindingType.cpp
@@ -174,20 +174,6 @@ protected:
 	TSharedPtr<IPropertyHandle> StructProperty;
 };
 
-FUniversalObjectLocatorBindingType::FUniversalObjectLocatorBindingType()
-{
-	OnObjectsReinstancedHandle = FCoreUObjectDelegates::OnObjectsReinstanced.AddLambda([this](const FCoreUObjectDelegates::FReplacementObjectMap& InMap)
-	{
-		// Empty the cached map on reinstancing
-		ClassMap.Empty();
-	});
-}
-
-FUniversalObjectLocatorBindingType::~FUniversalObjectLocatorBindingType()
-{
-	FCoreUObjectDelegates::OnObjectsReinstanced.Remove(OnObjectsReinstancedHandle);
-}
-
 TSharedRef<SWidget> FUniversalObjectLocatorBindingType::CreateEditWidget(const TSharedRef<IPropertyHandle>& InPropertyHandle, const FAnimNextParamType& InType) const
 {
 	class SEditWidget : public SCompoundWidget
@@ -694,7 +680,8 @@ FText FUniversalObjectLocatorBindingType::GetTooltipText(TConstStructView<FAnimN
 
 void FUniversalObjectLocatorBindingType::BuildBindingGraphFragment(const FRigVMCompileSettings& InSettings, const FBindingGraphFragmentArgs& InArgs, URigVMPin*& OutExecTail, FVector2D& OutLocation) const
 {
-	if(InArgs.Event != FRigUnit_AnimNextExecuteBindings::StaticStruct())
+	if( InArgs.Event != FRigUnit_AnimNextExecuteBindings_GT::StaticStruct() &&
+		InArgs.Event != FRigUnit_AnimNextExecuteBindings_WT::StaticStruct())
 	{
 		return;
 	}
@@ -712,13 +699,18 @@ void FUniversalObjectLocatorBindingType::BuildBindingGraphFragment(const FRigVMC
 	for(const FBindingGraphInput& Input : InArgs.Inputs)
 	{
 		const FAnimNextUniversalObjectLocatorBindingData& LocatorBinding = Input.BindingData.Get<FAnimNextUniversalObjectLocatorBindingData>();
-		FBindingData& BindingData = LocatorDataMap.FindOrAdd(LocatorBinding.Locator);
-		BindingData.Inputs.Add(&Input);
+
+		// Check if thread safety matches
+		if(LocatorBinding.IsThreadSafe() == InArgs.bThreadSafe)
+		{
+			FBindingData& BindingData = LocatorDataMap.FindOrAdd(LocatorBinding.Locator);
+			BindingData.Inputs.Add(&Input);
+		}
 	}
 
 	// Spawn object getter nodes for locators
 	float YOffset = 100.0f;
-	for(TPair<FUniversalObjectLocator, FBindingData> LocatorDataPair: LocatorDataMap)
+	for(const TPair<FUniversalObjectLocator, FBindingData>& LocatorDataPair : LocatorDataMap)
 	{
 		URigVMNode* ResolverNode = InArgs.Controller->AddUnitNode(FRigUnit_ResolveUniversalObjectLocator::StaticStruct(), FRigVMStruct::ExecuteName, FVector2D(0.0f, YOffset), TEXT(""), false);
 		if(ResolverNode == nullptr)
@@ -962,16 +954,6 @@ const UClass* FUniversalObjectLocatorBindingType::GetClass(TConstStructView<FAni
 	return LocatorEditor->ResolveClass(*LocatorBinding->Locator.GetLastFragment(), Context);
 }
 
-TSharedRef<FClassProxy> FUniversalObjectLocatorBindingType::FindOrCreateClassProxy(const UClass* InClass) const 
-{
-	if(TSharedPtr<FClassProxy>* FoundProxy = ClassMap.Find(InClass))
-	{
-		return FoundProxy->ToSharedRef();
-	}
-
-	return ClassMap.Add(InClass, MakeShared<FClassProxy>(InClass)).ToSharedRef();
-}
-
 }
 
 #undef LOCTEXT_NAMESPACE
\ No newline at end of file
diff --git a/Engine/Plugins/Experimental/AnimNext/Source/AnimNextUncookedOnly/Private/Variables/UniversalObjectLocatorBindingType.h b/Engine/Plugins/Experimental/AnimNext/Source/AnimNextUncookedOnly/Private/Variables/UniversalObjectLocatorBindingType.h
--- a/Engine/Plugins/Experimental/AnimNext/Source/AnimNextUncookedOnly/Private/Variables/UniversalObjectLocatorBindingType.h
+++ b/Engine/Plugins/Experimental/AnimNext/Source/AnimNextUncookedOnly/Private/Variables/UniversalObjectLocatorBindingType.h
@@ -16,10 +16,6 @@ namespace UE::AnimNext::UncookedOnly
 // Provides information about object proxy parameter sources
 class FUniversalObjectLocatorBindingType : public IVariableBindingType
 {
-public:
-	FUniversalObjectLocatorBindingType();
-	virtual ~FUniversalObjectLocatorBindingType() override;
-
 private:
 	// IVariableBindingType interface
 	virtual TSharedRef<SWidget> CreateEditWidget(const TSharedRef<IPropertyHandle>& InPropertyHandle, const FAnimNextParamType& InType) const override;
@@ -28,14 +24,6 @@ private:
 	virtual void BuildBindingGraphFragment(const FRigVMCompileSettings& InSettings, const FBindingGraphFragmentArgs& InArgs, URigVMPin*& OutExecTail, FVector2D& OutLocation) const override;
 
 	static const UClass* GetClass(TConstStructView<FAnimNextVariableBindingData> InBindingData);
-
-	TSharedRef<FClassProxy> FindOrCreateClassProxy(const UClass* InClass) const; 
-
-	// Cached map of classes -> proxy
-	mutable TMap<TObjectKey<UClass>, TSharedPtr<FClassProxy>> ClassMap;
-
-	// Hook for object reinstancing
-	FDelegateHandle OnObjectsReinstancedHandle;
 };
 
 }
diff --git a/Engine/Plugins/Experimental/AnimNext/Source/AnimNextUncookedOnly/Public/Variables/AnimNextVariableBindingData.h b/Engine/Plugins/Experimental/AnimNext/Source/AnimNextUncookedOnly/Public/Variables/AnimNextVariableBindingData.h
--- a/Engine/Plugins/Experimental/AnimNext/Source/AnimNextUncookedOnly/Public/Variables/AnimNextVariableBindingData.h
+++ b/Engine/Plugins/Experimental/AnimNext/Source/AnimNextUncookedOnly/Public/Variables/AnimNextVariableBindingData.h
@@ -15,4 +15,7 @@ struct FAnimNextVariableBindingData
 	
 	// Returns whether this binding data is valid or not
 	virtual bool IsValid() const { return false; }
+
+	// Returns whether this binding data is thread safe or not
+	virtual bool IsThreadSafe() const { return false; }
 };
diff --git a/Engine/Plugins/Experimental/AnimNext/Source/AnimNextUncookedOnly/Public/Variables/IVariableBindingType.h b/Engine/Plugins/Experimental/AnimNext/Source/AnimNextUncookedOnly/Public/Variables/IVariableBindingType.h
--- a/Engine/Plugins/Experimental/AnimNext/Source/AnimNextUncookedOnly/Public/Variables/IVariableBindingType.h
+++ b/Engine/Plugins/Experimental/AnimNext/Source/AnimNextUncookedOnly/Public/Variables/IVariableBindingType.h
@@ -28,23 +28,6 @@ namespace UE::AnimNext::Editor
 namespace UE::AnimNext::UncookedOnly
 {
 
-// Info about a variable binding gleaned from FindBindingInfo and ForEachBinding
-struct FVariableBindingInfo
-{
-	// The binding's type
-	FAnimNextParamType Type;
-	// Display name for editor
-	FText DisplayName;
-	// Tooltip to display in editor
-	FText Tooltip;
-	// Function used to access this binding, if any
-	const UFunction* Function = nullptr;
-	// Property for this binding, if any
-	const FProperty* Property = nullptr;
-	// Whether this binding is safe to be accessed on worker threads
-	bool bThreadSafe = false;
-};
-
 // Interface used in editor/uncooked situations to determine the characteristics of a variable binding
 class IVariableBindingType
 {
@@ -66,7 +49,7 @@ protected:
 
 	struct FBindingGraphFragmentArgs
 	{
-		// The event (e.g. FRigUnit_AnimNextExecuteBindings) that is currently being processed
+		// The event (e.g. FRigUnit_AnimNextExecuteBindings_GT) that is currently being processed
 		UScriptStruct* Event;
 		// All inputs, corresponding to variables
 		TConstArrayView<FBindingGraphInput> Inputs;
@@ -78,6 +61,8 @@ protected:
 		URigVMPin* ExecTail;
 		// The current spawn location, useful for making user-readable graphs
 		FVector2D CurrentLocation;
+		// Whether we should be processing thread-safe bindings
+		bool bThreadSafe;
 	};
 	
 private:
diff --git a/Engine/Plugins/Experimental/AnimNextAnimGraph/Config/BaseAnimNextAnimGraph.ini b/Engine/Plugins/Experimental/AnimNextAnimGraph/Config/BaseAnimNextAnimGraph.ini
--- a/Engine/Plugins/Experimental/AnimNextAnimGraph/Config/BaseAnimNextAnimGraph.ini
+++ b/Engine/Plugins/Experimental/AnimNextAnimGraph/Config/BaseAnimNextAnimGraph.ini
@@ -1,5 +1,6 @@
 ; Default settings/mappings for factories
 [/Script/AnimNextAnimGraph.AnimNextAnimGraphSettings]
+DefaultRunGraphHost=/AnimNextAnimGraph/FactoryGraphs/AG_DefaultGraphHost.AG_DefaultGraphHost
 +AssetGraphMappings=(AssetType="/Script/Engine.AnimSequence",AnimationGraph="/AnimNextAnimGraph/FactoryGraphs/AG_AnimSequencePlayer.AG_AnimSequencePlayer",Variable="AnimSequence")
 
 [CoreRedirects]
@@ -37,3 +38,13 @@
 ; Rename PlayAnim parameter
 +PropertyRedirects=(OldName="/Script/AnimNextAnimGraphUncookedOnly.CreateProxyObjectForPlayAnim.AnimationObject", NewName="AnimSequence")
 +PropertyRedirects=(OldName="/Script/AnimNextAnimGraphUncookedOnly.K2Node_PlayAnim.AnimationObject", NewName="AnimSequence")
+
+; PlayAnim -> Injection
++EnumRedirects=(OldName="/Script/AnimNextAnimGraph.EAnimNextPlayAnimBlendMode", NewName="/Script/AnimNextAnimGraph.EAnimNextInjectionBlendMode")
++EnumRedirects=(OldName="/Script/AnimNextAnimGraph.EAnimNextPlayAnimStatus", NewName="/Script/AnimNextAnimGraph.EAnimNextInjectionStatus")
++StructRedirects=(OldName="/Script/AnimNextAnimGraph.AnimNextPlayAnimBlendSettings", NewName="/Script/AnimNextAnimGraph.AnimNextInjectionBlendSettings")
++StructRedirects=(OldName="/Script/AnimNextAnimGraph.AnimNextPlayAnimRequestArgs", NewName="/Script/AnimNextAnimGraph.AnimNextPlayAnimRequestArgs")
++PropertyRedirects=(OldName="/Script/AnimNextAnimGraph.AnimNextInjectionRequestArgs.SlotName", NewName="Name")
++ClassRedirects=(OldName="/Script/AnimNextAnimGraphUncookedOnly.K2Node_AnimNextPlayAsset", NewName="/Script/AnimNextAnimGraphUncookedOnly.K2Node_AnimNextInjection")
++StructRedirects=(OldName="/Script/AnimNextAnimGraph.RigUnit_AnimNextRunAnimationGraph", NewName="/Script/AnimNextAnimGraph.RigUnit_AnimNextRunAnimationGraph_v1")
++StructRedirects=(OldName="/Script/AnimNextAnimGraph.AnimNextGraphInstancePtr", NewName="/Script/AnimNextAnimGraph.AnimNextAnimGraph")
diff --git a/Engine/Plugins/Experimental/AnimNextAnimGraph/Source/AnimNextAnimGraph/Internal/AnimNextAnimGraphSettings.h b/Engine/Plugins/Experimental/AnimNextAnimGraph/Source/AnimNextAnimGraph/Internal/AnimNextAnimGraphSettings.h
--- a/Engine/Plugins/Experimental/AnimNextAnimGraph/Source/AnimNextAnimGraph/Internal/AnimNextAnimGraphSettings.h
+++ b/Engine/Plugins/Experimental/AnimNextAnimGraph/Source/AnimNextAnimGraph/Internal/AnimNextAnimGraphSettings.h
@@ -2,6 +2,7 @@
 
 #pragma once
 
+#include "Graph/AnimNextAnimationGraph.h"
 #include "StructUtils/StructView.h"
 #include "Templates/SubclassOf.h"
 #include "UObject/ObjectKey.h"
@@ -66,6 +67,12 @@ public:
 	// Rebuild mappings (loading assets synchronously if required) for lookup
 	void LoadAndRebuildMappings();
 
+	// Get the animation graph we run by default when hosting in a module
+	const UAnimNextAnimationGraph* GetDefaultRunGraphHost() const
+	{
+		return DefaultRunGraphHost.Get();
+	}
+
 private:
 #if WITH_EDITOR
 	// UObject interface
@@ -78,6 +85,10 @@ private:
 	TArray<UClass*> GetAllowedAssetClassesImpl() const;
 	
 private:
+	// The animation graph we run by default when hosting in a module
+	UPROPERTY(Config, EditAnywhere, Category = "Graphs")
+	TSoftObjectPtr<const UAnimNextAnimationGraph> DefaultRunGraphHost;
+
 	// Mappings from assets to animation graphs
 	UPROPERTY(Config, EditAnywhere, DisplayName = "Asset/Graph Mappings", Category = "Mappings")
 	TArray<FAnimNextAssetGraphMapping> AssetGraphMappings;
diff --git a/Engine/Plugins/Experimental/AnimNextAnimGraph/Source/AnimNextAnimGraph/Internal/Graph/AnimNextModuleAnimGraphComponent.h b/Engine/Plugins/Experimental/AnimNextAnimGraph/Source/AnimNextAnimGraph/Internal/Graph/AnimNextModuleAnimGraphComponent.h
new file mode 100644
--- /dev/null
+++ b/Engine/Plugins/Experimental/AnimNextAnimGraph/Source/AnimNextAnimGraph/Internal/Graph/AnimNextModuleAnimGraphComponent.h
@@ -0,0 +1,41 @@
+// Copyright Epic Games, Inc. All Rights Reserved.
+
+#pragma once
+
+#include "Injection/InjectionInfo.h"
+#include "Module/ModuleInstanceComponent.h"
+#include "AnimNextModuleAnimGraphComponent.generated.h"
+
+struct FAnimNextGraphInstance;
+class UAnimNextAnimationGraph;
+
+// Module component that owns/allocates/release all animation graph instances on a module
+USTRUCT()
+struct FAnimNextModuleAnimGraphComponent : public FAnimNextModuleInstanceComponent
+{
+	GENERATED_BODY()
+
+	ANIMNEXT_DECLARE_MODULE_INSTANCE_COMPONENT(FAnimNextModuleAnimGraphComponent);
+
+	FAnimNextModuleAnimGraphComponent() = default;
+	explicit FAnimNextModuleAnimGraphComponent(FAnimNextModuleInstance& InOwnerInstance);
+
+	TWeakPtr<FAnimNextGraphInstance> AllocateInstance(const UAnimNextAnimationGraph* InAnimationGraph, FAnimNextGraphInstance* InParentInstance = nullptr, FName InEntryPoint = NAME_None);
+
+	void ReleaseInstance(TWeakPtr<FAnimNextGraphInstance> InInstance);
+
+	void AddStructReferencedObjects(class FReferenceCollector& Collector);
+
+private:
+	// All the owned graph instances for the module
+	TArray<TSharedPtr<FAnimNextGraphInstance>> GraphInstances;
+};
+
+template<>
+struct TStructOpsTypeTraits<FAnimNextModuleAnimGraphComponent> : public TStructOpsTypeTraitsBase2<FAnimNextModuleAnimGraphComponent>
+{
+	enum
+	{
+		WithAddStructReferencedObjects = true,
+	};
+};
diff --git a/Engine/Plugins/Experimental/AnimNextAnimGraph/Source/AnimNextAnimGraph/Internal/Graph/RigUnit_AnimNextGraphEvaluator.h b/Engine/Plugins/Experimental/AnimNextAnimGraph/Source/AnimNextAnimGraph/Internal/Graph/RigUnit_AnimNextGraphEvaluator.h
--- a/Engine/Plugins/Experimental/AnimNextAnimGraph/Source/AnimNextAnimGraph/Internal/Graph/RigUnit_AnimNextGraphEvaluator.h
+++ b/Engine/Plugins/Experimental/AnimNextAnimGraph/Source/AnimNextAnimGraph/Internal/Graph/RigUnit_AnimNextGraphEvaluator.h
@@ -6,34 +6,7 @@
 #include "AnimNextExecuteContext.h"
 #include "RigUnit_AnimNextGraphEvaluator.generated.h"
 
-// Represents an argument to a FRigUnit_AnimNextGraphEvaluator entry point
-USTRUCT()
-struct FAnimNextGraphEvaluatorExecuteArgument
-{
-	GENERATED_BODY();
-
-	UPROPERTY()
-	FString Name;
-
-	UPROPERTY()
-	FString CPPType;
-};
-
-// Represents an entry point of FRigUnit_AnimNextGraphEvaluator
-USTRUCT()
-struct FAnimNextGraphEvaluatorExecuteDefinition
-{
-	GENERATED_BODY();
-
-	UPROPERTY()
-	uint32 Hash = 0;
-
-	UPROPERTY()
-	FString MethodName;
-
-	UPROPERTY()
-	TArray<FAnimNextGraphEvaluatorExecuteArgument> Arguments;
-};
+struct FAnimNextGraphEvaluatorExecuteDefinition;
 
 /**
  * Animation graph evaluator
diff --git a/Engine/Plugins/Experimental/AnimNextAnimGraph/Source/AnimNextAnimGraph/Internal/Injection/InjectionCallbackProxy.h b/Engine/Plugins/Experimental/AnimNextAnimGraph/Source/AnimNextAnimGraph/Internal/Injection/InjectionCallbackProxy.h
new file mode 100644
--- /dev/null
+++ b/Engine/Plugins/Experimental/AnimNextAnimGraph/Source/AnimNextAnimGraph/Internal/Injection/InjectionCallbackProxy.h
@@ -0,0 +1,81 @@
+﻿// Copyright Epic Games, Inc. All Rights Reserved.
+
+#pragma once
+
+#include "CoreMinimal.h"
+#include "Engine/CancellableAsyncAction.h"
+#include "Injection/InjectionRequest.h"
+
+#include "InjectionCallbackProxy.generated.h"
+
+class UAnimNextComponent;
+
+DECLARE_DYNAMIC_MULTICAST_DELEGATE(FOnInjectionDelegate);
+
+UENUM()
+enum class EUninjectionResult : uint8
+{
+	Succeeded,
+	Failed
+};
+
+UCLASS(MinimalAPI, meta=(ExposedAsyncProxy = "AsyncTask", HasDedicatedAsyncNode))
+class UInjectionCallbackProxy : public UCancellableAsyncAction
+{
+	GENERATED_BODY()
+
+public:
+	// Called when the provided animation object finished playing and hasn't been interrupted
+	UPROPERTY(BlueprintAssignable)
+	FOnInjectionDelegate OnCompleted;
+
+	// Called when the provided animation object starts blending out and hasn't been interrupted
+	UPROPERTY(BlueprintAssignable)
+	FOnInjectionDelegate OnBlendOut;
+
+	// Called when the provided animation object has been interrupted (or failed to play)
+	UPROPERTY(BlueprintAssignable)
+	FOnInjectionDelegate OnInterrupted;
+
+	// Called to perform the query internally
+	UFUNCTION(BlueprintCallable, meta = (BlueprintInternalUseOnly = "true"))
+	static ANIMNEXTANIMGRAPH_API UInjectionCallbackProxy* CreateProxyObjectForInjection(
+		UAnimNextComponent* AnimNextComponent,
+		FName SiteName,
+		UObject* Object,
+		UAnimNextComponent* BindingComponent,
+		FInstancedStruct Payload,
+		FAnimNextInjectionBlendSettings BlendInSettings = FAnimNextInjectionBlendSettings(),
+		FAnimNextInjectionBlendSettings BlendOutSettings = FAnimNextInjectionBlendSettings());
+
+	// Un-inject a previously injected object. Cancelling this async tasks will also un-inject.
+	UFUNCTION(BlueprintCallable, Category="Animation|AnimNext", meta=(ExpandEnumAsExecs = "ReturnValue"))
+	ANIMNEXTANIMGRAPH_API EUninjectionResult Uninject();
+
+	// UObject Interface
+	ANIMNEXTANIMGRAPH_API virtual void BeginDestroy() override;
+
+	// UCancellableAsyncAction interface
+	ANIMNEXTANIMGRAPH_API virtual void Cancel() override;
+
+protected:
+	ANIMNEXTANIMGRAPH_API void OnInjectionCompleted(const UE::AnimNext::FInjectionRequest& Request);
+	ANIMNEXTANIMGRAPH_API void OnInjectionInterrupted(const UE::AnimNext::FInjectionRequest& Request);
+	ANIMNEXTANIMGRAPH_API void OnInjectionBlendingOut(const UE::AnimNext::FInjectionRequest& Request);
+
+	// Attempts to play an object with the specified payload. Returns whether it started or not.
+	ANIMNEXTANIMGRAPH_API bool Inject(
+		UAnimNextComponent* AnimNextComponent,
+		FName SiteName,
+		UObject* Object,
+		UAnimNextComponent* BindingComponent,
+		FInstancedStruct&& Payload,
+		const UE::AnimNext::FInjectionBlendSettings& BlendInSettings,
+		const UE::AnimNext::FInjectionBlendSettings& BlendOutSettings);
+
+private:
+	void Reset();
+
+	UE::AnimNext::FInjectionRequestPtr PlayingRequest;
+	bool bWasInterrupted = false;
+};
