AnimNext: Compile button, auto compile, clickable errors and node error highlights

diff --git a/Engine/Plugins/Experimental/AnimNext/Source/AnimNextEditor/AnimNextEditor.Build.cs b/Engine/Plugins/Experimental/AnimNext/Source/AnimNextEditor/AnimNextEditor.Build.cs
--- a/Engine/Plugins/Experimental/AnimNext/Source/AnimNextEditor/AnimNextEditor.Build.cs
+++ b/Engine/Plugins/Experimental/AnimNext/Source/AnimNextEditor/AnimNextEditor.Build.cs
@@ -49,6 +49,9 @@ namespace UnrealBuildTool.Rules
 					"Kismet",
 					"AdvancedWidgets", 
 					"SceneOutliner",
+					"EditorFramework",
+					"InteractiveToolsFramework",
+					"EditorInteractiveToolsFramework",
 				}
 			);
 		}
diff --git a/Engine/Plugins/Experimental/AnimNext/Source/AnimNextEditor/Private/AnimNextEditorContext.h b/Engine/Plugins/Experimental/AnimNext/Source/AnimNextEditor/Private/AnimNextEditorContext.h
new file mode 100644
--- /dev/null
+++ b/Engine/Plugins/Experimental/AnimNext/Source/AnimNextEditor/Private/AnimNextEditorContext.h
@@ -0,0 +1,22 @@
+﻿// Copyright Epic Games, Inc. All Rights Reserved.
+
+#pragma once
+
+#include "EdMode.h"
+
+#include "AnimNextEditorContext.generated.h"
+
+namespace UE::Workspace
+{
+	class IWorkspaceEditor;
+}
+
+UCLASS(Transient)
+class UAnimNextEditorContext : public UObject
+{
+	GENERATED_BODY()
+
+public:
+	// The workspace editor host
+	TWeakPtr<UE::Workspace::IWorkspaceEditor> WeakWorkspaceEditor;
+};
diff --git a/Engine/Plugins/Experimental/AnimNext/Source/AnimNextEditor/Private/AnimNextEditorModule.cpp b/Engine/Plugins/Experimental/AnimNext/Source/AnimNextEditor/Private/AnimNextEditorModule.cpp
--- a/Engine/Plugins/Experimental/AnimNext/Source/AnimNextEditor/Private/AnimNextEditorModule.cpp
+++ b/Engine/Plugins/Experimental/AnimNext/Source/AnimNextEditor/Private/AnimNextEditorModule.cpp
@@ -34,15 +34,23 @@
 #include "Param/ObjectPropertyLocatorEditor.h"
 #include "Variables/SAddVariablesDialog.h"
 #include "AnimNextAssetWorkspaceAssetUserData.h"
+#include "AnimNextEditorContext.h"
+#include "AssetCompilationHandler.h"
+#include "ContextObjectStore.h"
+#include "EditorModeManager.h"
 #include "Common/AnimNextGraphItemDetails.h"
 #include "Common/AnimNextCollapseNodeItemDetails.h"
 #include "Common/AnimNextFunctionItemDetails.h"
 #include "Param/AnimNextActorLocatorEditor.h"
 #include "IWorkspaceEditor.h"
+#include "RigVMCommands.h"
+#include "ToolMenus.h"
 #include "Framework/Docking/LayoutExtender.h"
 #include "Common/AnimNextAssetItemDetails.h"
 #include "Common/AnimNextRigVMAssetEditorDataCustomization.h"
 #include "DataInterface/AnimNextDataInterface.h"
+#include "Editor/RigVMEditorStyle.h"
+#include "EdModeInteractiveToolsContext.h"
 #include "Entries/AnimNextEventGraphEntry.h"
 #include "Logging/MessageLog.h"
 #include "Graph/AnimNextEdGraphCustomization.h"
@@ -50,11 +58,13 @@
 #include "Variables/SVariablesView.h"
 #include "Variables/VariableOverrideCommands.h"
 #include "Variables/VariableProxyCustomization.h"
+#include "Workspace/AnimNextWorkspaceEditorMode.h"
 #include "Workspace/AnimNextWorkspaceSchema.h"
 #include "ExternalPackageHelper.h"
 #include "Common/GraphEditorSchemaActions.h"
 #include "EditorUtils.h"
 #include "PersonaModule.h"
+#include "Widgets/Docking/SDockTab.h"
 
 #define LOCTEXT_NAMESPACE "AnimNextEditorModule"
 
@@ -64,6 +74,7 @@ namespace UE::AnimNext::Editor
 void FAnimNextEditorModule::StartupModule()
 {
 	FVariableOverrideCommands::Register();
+	FRigVMCommands::Register();
 
 	// Register settings for user editing
 	ISettingsModule& SettingsModule = FModuleManager::Get().LoadModuleChecked<ISettingsModule>("Settings");
@@ -114,6 +125,125 @@ void FAnimNextEditorModule::StartupModule()
 		InLayoutExtender.ExtendLayout(FTabId(Workspace::WorkspaceTabs::BottomLeftDocumentArea), ELayoutExtensionPosition::After, VariablesTab);
 	});
 
+	WorkspaceEditorModule.OnExtendToolMenuContext().AddLambda([](TSharedPtr<UE::Workspace::IWorkspaceEditor> InWorkspaceEditor, FToolMenuContext& InContext)
+	{
+		UContextObjectStore* ContextStore = InWorkspaceEditor->GetEditorModeManager().GetInteractiveToolsContext()->ContextObjectStore;
+		UAnimNextEditorContext* AnimNextEditorContext = ContextStore->FindContext<UAnimNextEditorContext>();
+		if (!AnimNextEditorContext)
+		{
+			AnimNextEditorContext = NewObject<UAnimNextEditorContext>();
+			AnimNextEditorContext->WeakWorkspaceEditor = InWorkspaceEditor;
+			ContextStore->AddContextObject(AnimNextEditorContext);
+		}
+		
+		if (UToolMenu* Menu = UToolMenus::Get()->ExtendMenu(InWorkspaceEditor->GetToolMenuToolbarName()))
+		{
+			FToolMenuSection& RigVMOperationsSection = Menu->AddSection("RigVMOperations", TAttribute<FText>(), FToolMenuInsert("WorkspaceOperations", EToolMenuInsertType::After));
+
+			auto GetIcon = [WeakWorkspaceEditor = TWeakPtr<UE::Workspace::IWorkspaceEditor>(InWorkspaceEditor)]()
+			{
+				static const FName CompileStatusBackground("AssetEditor.CompileStatus.Background");
+				static const FName CompileStatusUnknown("AssetEditor.CompileStatus.Overlay.Unknown");
+				static const FName CompileStatusError("AssetEditor.CompileStatus.Overlay.Error");
+				static const FName CompileStatusGood("AssetEditor.CompileStatus.Overlay.Good");
+				static const FName CompileStatusWarning("AssetEditor.CompileStatus.Overlay.Warning");
+
+				TSharedPtr<UE::Workspace::IWorkspaceEditor> WorkspaceEditor = WeakWorkspaceEditor.Pin();
+				if(!WorkspaceEditor.IsValid())
+				{
+					return FSlateIcon();
+				}
+
+				UAnimNextWorkspaceEditorMode* EditorMode = Cast<UAnimNextWorkspaceEditorMode>(WorkspaceEditor->GetEditorModeManager().GetActiveScriptableMode(UAnimNextWorkspaceEditorMode::EM_AnimNextWorkspace));
+				if(EditorMode == nullptr)
+				{
+					return FSlateIcon(FAppStyle::Get().GetStyleSetName(), CompileStatusBackground, NAME_None, CompileStatusUnknown);
+				}
+
+				ECompileStatus Status = EditorMode->GetLatestCompileStatus();
+				switch (Status)
+				{
+				default:
+				case ECompileStatus::Unknown:
+				case ECompileStatus::Dirty:
+					return FSlateIcon(FAppStyle::GetAppStyleSetName(), CompileStatusBackground, NAME_None, CompileStatusUnknown);
+				case ECompileStatus::Error:
+					return FSlateIcon(FAppStyle::GetAppStyleSetName(), CompileStatusBackground, NAME_None, CompileStatusError);
+				case ECompileStatus::UpToDate:
+					return FSlateIcon(FAppStyle::GetAppStyleSetName(), CompileStatusBackground, NAME_None, CompileStatusGood);
+				case ECompileStatus::Warning:
+					return FSlateIcon(FAppStyle::GetAppStyleSetName(), CompileStatusBackground, NAME_None, CompileStatusWarning);
+				}
+			};
+
+			auto GetTooltip = [WeakWorkspaceEditor = TWeakPtr<UE::Workspace::IWorkspaceEditor>(InWorkspaceEditor)]()
+			{
+				TSharedPtr<UE::Workspace::IWorkspaceEditor> WorkspaceEditor = WeakWorkspaceEditor.Pin();
+				if(!WorkspaceEditor.IsValid())
+				{
+					return LOCTEXT("CompileGenericTooltip", "Compile all relevant assets");
+				}
+
+				UAnimNextWorkspaceEditorMode* EditorMode = Cast<UAnimNextWorkspaceEditorMode>(WorkspaceEditor->GetEditorModeManager().GetActiveScriptableMode(UAnimNextWorkspaceEditorMode::EM_AnimNextWorkspace));
+				if(EditorMode == nullptr)
+				{
+					return LOCTEXT("CompileGenericTooltip", "Compile all relevant assets");
+				}
+
+				FTextBuilder TextBuilder;
+				if(EditorMode->GetState().bCompileDirtyFiles)
+				{
+					if(EditorMode->GetState().bCompileWholeWorkspace)
+					{
+						TextBuilder.AppendLine(LOCTEXT("CompileWorkspaceDirtyTooltip", "Compile all assets in the workspace that are dirty or have errors"));
+					}
+					else
+					{
+						TextBuilder.AppendLine(LOCTEXT("CompileDirtyTooltip", "Compile the current asset if it is dirty or has errors"));
+					}
+				}
+				else
+				{
+					if(EditorMode->GetState().bCompileWholeWorkspace)
+					{
+						TextBuilder.AppendLine(LOCTEXT("CompileWorkspaceTooltip", "Compile all assets in the workspace"));
+					}
+					else
+					{
+						TextBuilder.AppendLine(LOCTEXT("CompileCurrentTooltip", "Compile the current asset"));
+					}
+				}
+
+				if(EditorMode->GetState().bAutoCompile)
+				{
+					TextBuilder.AppendLine(LOCTEXT("AutoCompileEnabledTooltip", "Auto-compilation is enabled"));
+				}
+				else
+				{
+					TextBuilder.AppendLine(LOCTEXT("AutoCompileEnabledTooltip", "Auto-compilation is disabled"));
+				}
+
+				return TextBuilder.ToText();
+			};
+
+			RigVMOperationsSection.AddEntry(FToolMenuEntry::InitToolBarButton(FRigVMCommands::Get().Compile, TAttribute<FText>(), MakeAttributeLambda(GetTooltip), MakeAttributeLambda(GetIcon)));
+			RigVMOperationsSection.AddEntry(FToolMenuEntry::InitComboButton(
+				"CompileOptionsCombo",
+				FUIAction(),
+				FNewToolMenuDelegate::CreateLambda([](UToolMenu* InToolMenu)
+				{
+					FToolMenuSection& CompileOptionsSection = InToolMenu->AddSection("CompileOptionsSection", LOCTEXT("CompileOptionsComboLabel", "Compilation Options"));
+					CompileOptionsSection.AddEntry(FToolMenuEntry::InitMenuEntry(FRigVMCommands::Get().AutoCompile, TAttribute<FText>(), TAttribute<FText>(), FSlateIcon("RigVMEditorStyle", "RigVM.AutoCompileGraph")));
+					CompileOptionsSection.AddEntry(FToolMenuEntry::InitMenuEntry(FRigVMCommands::Get().CompileWholeWorkspace, TAttribute<FText>(), TAttribute<FText>(), FSlateIcon(FAppStyle::GetAppStyleSetName(), "Icons.World")));
+					CompileOptionsSection.AddEntry(FToolMenuEntry::InitMenuEntry(FRigVMCommands::Get().CompileDirtyFiles, TAttribute<FText>(), TAttribute<FText>(), FSlateIcon(FAppStyle::GetAppStyleSetName(), "Icons.DirtyBadge")));
+				}),
+				FText::GetEmpty(),
+				LOCTEXT("CompileOptionsComboTooltip", "Compilation Options"),
+				FSlateIcon(),
+				true));
+		}
+	});
+	
 	RegisterWorkspaceDocumentTypes(WorkspaceEditorModule);
 
 	WorkspaceEditorModule.OnRegisterWorkspaceDetailsCustomization().AddLambda([](const TWeakPtr<Workspace::IWorkspaceEditor>& InWorkspaceEditor, TSharedPtr<IDetailsView>& InDetailsView)
@@ -191,6 +321,11 @@ void FAnimNextEditorModule::StartupModule()
 	NotifyHostParameters.OnRemoveNotify = FPersonaModule::FNotifyHostAssetParameters::FOnRemoveNotify::CreateStatic(&UAnimNextRigVMAssetEditorData::HandleRemoveNotify);
 	NotifyHostParameters.OnReplaceNotify = FPersonaModule::FNotifyHostAssetParameters::FOnReplaceNotify::CreateStatic(&UAnimNextRigVMAssetEditorData::HandleReplaceNotify);
 	PersonaModule.RegisterNotifyHostAsset(FTopLevelAssetPath(TEXT("/Script/AnimNext.AnimNextRigVMAsset")), NotifyHostParameters);
+
+	RegisterAssetCompilationHandler(FTopLevelAssetPath(TEXT("/Script/AnimNext.AnimNextRigVMAsset")), FAssetCompilationHandlerFactoryDelegate::CreateLambda([](UObject* InAsset)
+	{
+		return MakeShared<FAssetCompilationHandler>(InAsset);
+	}));
 }
 
 void FAnimNextEditorModule::ShutdownModule()
@@ -245,6 +380,7 @@ void FAnimNextEditorModule::ShutdownModule()
 		FPersonaModule& PersonaModule = FModuleManager::GetModuleChecked<FPersonaModule>("Persona");
 		PersonaModule.UnregisterNotifyHostAsset(FTopLevelAssetPath(TEXT("/Script/AnimNext.AnimNextRigVMAsset")));
 	}
+	UnregisterAssetCompilationHandler(FTopLevelAssetPath(TEXT("/Script/AnimNext.AnimNextModule")));
 }
 
 void FAnimNextEditorModule::RegisterLocatorFragmentEditorType(FName InLocatorFragmentEditorName)
@@ -283,6 +419,17 @@ void FAnimNextEditorModule::UnregisterGraphMenuActionsProvider(const FDelegateHa
 	OnCollectGraphMenuActionsDelegateImpl.Remove(DelegateHandle);
 }
 
+void FAnimNextEditorModule::RegisterAssetCompilationHandler(const FTopLevelAssetPath& InClassPath, FAssetCompilationHandlerFactoryDelegate InAssetCompilationHandlerFactory)
+{
+	check(InAssetCompilationHandlerFactory.IsBound());
+	AssetCompilationHandlerFactories.Add(InClassPath, InAssetCompilationHandlerFactory);
+}
+
+void FAnimNextEditorModule::UnregisterAssetCompilationHandler(const FTopLevelAssetPath& InClassPath)
+{
+	AssetCompilationHandlerFactories.Remove(InClassPath);
+}
+
 void FAnimNextEditorModule::CollectGraphMenuActions(const TWeakPtr<UE::Workspace::IWorkspaceEditor>& WorkspaceEditor, FGraphContextMenuBuilder& InContexMenuBuilder, const FActionMenuContextData& InActionMenuContextData)
 {
 	if (OnCollectGraphMenuActionsDelegateImpl.IsBound())
@@ -443,6 +590,8 @@ void FAnimNextEditorModule::RegisterWorkspaceDocumentTypes(Workspace::IWorkspace
 				return nullptr;
 			}));
 
+	AnimNextModuleDocumentArgs.DocumentEditorMode = UAnimNextWorkspaceEditorMode::EM_AnimNextWorkspace;
+
 	WorkspaceEditorModule.RegisterObjectDocumentType(FTopLevelAssetPath(TEXT("/Script/AnimNext.AnimNextModule")), AnimNextModuleDocumentArgs);
 	
 	// --- AnimNextEdGraph ---
@@ -975,7 +1124,30 @@ void FAnimNextEditorModule::RegisterWorkspaceDocumentTypes(Workspace::IWorkspace
 		}
 	});
 
+	GraphDocumentArgs.DocumentEditorMode = UAnimNextWorkspaceEditorMode::EM_AnimNextWorkspace;
+
 	WorkspaceEditorModule.RegisterObjectDocumentType(FTopLevelAssetPath(TEXT("/Script/AnimNextUncookedOnly.AnimNextEdGraph")), GraphDocumentArgs);
+
+	Workspace::FDocumentSubObjectArgs GraphNodeSubObjectArgs;
+	GraphNodeSubObjectArgs.OnGetDocumentForSubObject = Workspace::FOnGetDocumentForSubObject::CreateLambda([](const UObject* InObject) -> UObject*
+	{
+		if(const UAnimNextEdGraphNode* EdGraphNode = Cast<UAnimNextEdGraphNode>(InObject))
+		{
+			return EdGraphNode->GetTypedOuter<UAnimNextEdGraph>();
+		}
+		return nullptr;
+	});
+	GraphNodeSubObjectArgs.OnPostDocumentOpenedForSubObject = Workspace::FOnPostDocumentOpenedForSubObject::CreateLambda([](const Workspace::FWorkspaceEditorContext& InContext, TSharedRef<SWidget> InWidget, UObject* InObject)
+	{
+		if(UAnimNextEdGraphNode* EdGraphNode = Cast<UAnimNextEdGraphNode>(InObject))
+		{
+			check(InWidget->GetType() == SGraphEditor::StaticWidgetClass().GetWidgetType());
+			TSharedPtr<SGraphEditor> GraphEditor = StaticCastSharedRef<SGraphEditor>(InWidget);
+			GraphEditor->JumpToNode(EdGraphNode, false);
+		}
+	});
+
+	WorkspaceEditorModule.RegisterDocumentSubObjectType(FTopLevelAssetPath(TEXT("/Script/AnimNextUncookedOnly.AnimNextEdGraphNode")), GraphNodeSubObjectArgs);
 }
 
 void FAnimNextEditorModule::UnregisterWorkspaceDocumentTypes()
@@ -985,9 +1157,25 @@ void FAnimNextEditorModule::UnregisterWorkspaceDocumentTypes()
 		Workspace::IWorkspaceEditorModule& WorkspaceEditorModule = FModuleManager::LoadModuleChecked<Workspace::IWorkspaceEditorModule>("WorkspaceEditor");
 		WorkspaceEditorModule.UnregisterObjectDocumentType(FTopLevelAssetPath(TEXT("/Script/AnimNext.AnimNextModule")));
 		WorkspaceEditorModule.UnregisterObjectDocumentType(FTopLevelAssetPath(TEXT("/Script/AnimNextUncookedOnly.AnimNextEdGraph")));
+		WorkspaceEditorModule.UnregisterDocumentSubObjectType(FTopLevelAssetPath(TEXT("/Script/AnimNextUncookedOnly.AnimNextEdGraphNode")));
 	}
 }
 
+const FAssetCompilationHandlerFactoryDelegate* FAnimNextEditorModule::FindAssetCompilationHandlerFactory(UClass* InAssetClass) const
+{
+	UClass* AssetClass = InAssetClass;
+	while(AssetClass)
+	{
+		if(const FAssetCompilationHandlerFactoryDelegate* FoundDelegate = AssetCompilationHandlerFactories.Find(AssetClass->GetClassPathName()))
+		{
+			return FoundDelegate;
+		}
+		AssetClass = AssetClass->GetSuperClass();
+	}
+
+	return nullptr;
+}
+
 }
 
 IMPLEMENT_MODULE(UE::AnimNext::Editor::FAnimNextEditorModule, AnimNextEditor);
diff --git a/Engine/Plugins/Experimental/AnimNext/Source/AnimNextEditor/Private/AnimNextEditorModule.h b/Engine/Plugins/Experimental/AnimNext/Source/AnimNextEditor/Private/AnimNextEditorModule.h
--- a/Engine/Plugins/Experimental/AnimNext/Source/AnimNextEditor/Private/AnimNextEditorModule.h
+++ b/Engine/Plugins/Experimental/AnimNext/Source/AnimNextEditor/Private/AnimNextEditorModule.h
@@ -7,6 +7,7 @@
 
 class UAnimNextWorkspaceSchema;
 class FAnimNextGraphPanelNodeFactory;
+class UAnimNextWorkspaceEditorMode;
 
 namespace UE::AnimNext::Editor
 {
@@ -25,6 +26,8 @@ namespace UE::AnimNext::Editor
 class FAnimNextEditorModule : public IAnimNextEditorModule
 {
 	friend UAnimNextWorkspaceSchema;
+	friend UAnimNextWorkspaceEditorMode;
+
 private:
 	virtual void StartupModule() override;
 	virtual void ShutdownModule() override;
@@ -36,6 +39,8 @@ private:
 	virtual void RemoveWorkspaceSupportedAssetClass(const FTopLevelAssetPath& InClassAssetPath) override;
 	virtual FDelegateHandle RegisterGraphMenuActionsProvider(const FOnCollectGraphMenuActionsDelegate& InCollectDelegate) override;
 	virtual void UnregisterGraphMenuActionsProvider(const FDelegateHandle& InDelegateHandle) override;
+	virtual void RegisterAssetCompilationHandler(const FTopLevelAssetPath& InClassPath, FAssetCompilationHandlerFactoryDelegate InAssetCompilationHandlerFactory) override;
+	virtual void UnregisterAssetCompilationHandler(const FTopLevelAssetPath& InClassPath) override;
 
 private:
 	void CollectGraphMenuActions(const TWeakPtr<UE::Workspace::IWorkspaceEditor>& WorkspaceEditor, FGraphContextMenuBuilder& InContexMenuBuilder, const FActionMenuContextData& InActionMenuContextData);
@@ -43,6 +48,8 @@ private:
 	void RegisterWorkspaceDocumentTypes(Workspace::IWorkspaceEditorModule& WorkspaceEditorModule);
 	void UnregisterWorkspaceDocumentTypes();
 
+	const FAssetCompilationHandlerFactoryDelegate* FindAssetCompilationHandlerFactory(UClass* InAssetClass) const;
+
 private:
 	/** Type identifier for parameter names */
 	TSharedPtr<FParamNamePropertyTypeIdentifier> Identifier;
@@ -54,6 +61,8 @@ private:
 
 	FCollectGraphMenuActionsMulticast OnCollectGraphMenuActionsDelegateImpl;
 
+	TMap<FTopLevelAssetPath, FAssetCompilationHandlerFactoryDelegate> AssetCompilationHandlerFactories;
+
 	friend class FLocatorContext;
 };
 
diff --git a/Engine/Plugins/Experimental/AnimNext/Source/AnimNextEditor/Private/AssetCompilationHandler.cpp b/Engine/Plugins/Experimental/AnimNext/Source/AnimNextEditor/Private/AssetCompilationHandler.cpp
new file mode 100644
--- /dev/null
+++ b/Engine/Plugins/Experimental/AnimNext/Source/AnimNextEditor/Private/AssetCompilationHandler.cpp
@@ -0,0 +1,71 @@
+﻿// Copyright Epic Games, Inc. All Rights Reserved.
+
+#include "AssetCompilationHandler.h"
+#include "AnimNextRigVMAsset.h"
+#include "AnimNextRigVMAssetEditorData.h"
+#include "AnimNextScopedCompilerResults.h"
+#include "UncookedOnlyUtils.h"
+
+namespace UE::AnimNext::Editor
+{
+
+FAssetCompilationHandler::FAssetCompilationHandler(UObject* InAsset)
+{
+	UAnimNextRigVMAsset* AnimNextRigVMAsset = CastChecked<UAnimNextRigVMAsset>(InAsset);
+	UAnimNextRigVMAssetEditorData* EditorData = UE::AnimNext::UncookedOnly::FUtils::GetEditorData<UAnimNextRigVMAssetEditorData>(AnimNextRigVMAsset);
+
+	NumErrors = EditorData->bErrorsDuringCompilation ? 1 : 0;
+	NumWarnings = EditorData->bWarningsDuringCompilation ? 1 : 0;
+}
+
+void FAssetCompilationHandler::Compile(TSharedRef<Workspace::IWorkspaceEditor> InWorkspaceEditor, UObject* InAsset)
+{
+	UAnimNextRigVMAsset* AnimNextRigVMAsset = CastChecked<UAnimNextRigVMAsset>(InAsset);
+	UAnimNextRigVMAssetEditorData* EditorData = UE::AnimNext::UncookedOnly::FUtils::GetEditorData<UAnimNextRigVMAssetEditorData>(AnimNextRigVMAsset);
+
+	UncookedOnly::FScopedCompilerResults CompilerResults(InAsset);
+	EditorData->RecompileVM();
+
+	NumErrors = CompilerResults.GetLog().NumErrors;
+	NumWarnings = CompilerResults.GetLog().NumWarnings;
+
+	OnCompileStatusChanged().ExecuteIfBound();
+}
+
+void FAssetCompilationHandler::SetAutoCompile(TSharedRef<Workspace::IWorkspaceEditor> InWorkspaceEditor, UObject* InAsset, bool bInAutoCompile)
+{
+	UAnimNextRigVMAsset* AnimNextRigVMAsset = CastChecked<UAnimNextRigVMAsset>(InAsset);
+	UAnimNextRigVMAssetEditorData* EditorData = UE::AnimNext::UncookedOnly::FUtils::GetEditorData<UAnimNextRigVMAssetEditorData>(AnimNextRigVMAsset);
+	EditorData->SetAutoVMRecompile(bInAutoCompile);
+}
+
+bool FAssetCompilationHandler::GetAutoCompile(TSharedRef<Workspace::IWorkspaceEditor> InWorkspaceEditor, const UObject* InAsset) const
+{
+	const UAnimNextRigVMAsset* AnimNextRigVMAsset = CastChecked<UAnimNextRigVMAsset>(InAsset);
+	UAnimNextRigVMAssetEditorData* EditorData = UE::AnimNext::UncookedOnly::FUtils::GetEditorData<UAnimNextRigVMAssetEditorData>(AnimNextRigVMAsset);
+	return EditorData->GetAutoVMRecompile();
+}
+
+ECompileStatus FAssetCompilationHandler::GetCompileStatus(TSharedRef<Workspace::IWorkspaceEditor> InWorkspaceEditor, const UObject* InAsset) const
+{
+	const UAnimNextRigVMAsset* AnimNextRigVMAsset = CastChecked<UAnimNextRigVMAsset>(InAsset);
+	UAnimNextRigVMAssetEditorData* EditorData = UE::AnimNext::UncookedOnly::FUtils::GetEditorData<UAnimNextRigVMAssetEditorData>(AnimNextRigVMAsset);
+	if(EditorData->bErrorsDuringCompilation || NumErrors > 0)
+	{
+		return ECompileStatus::Error;
+	}
+
+	if(EditorData->bWarningsDuringCompilation || NumWarnings > 0)
+	{
+		return ECompileStatus::Warning;
+	}
+
+	if(EditorData->IsDirtyForRecompilation())
+	{
+		return ECompileStatus::Dirty;
+	}
+
+	return ECompileStatus::UpToDate;
+}
+
+}
diff --git a/Engine/Plugins/Experimental/AnimNext/Source/AnimNextEditor/Private/Common/AnimNextCompilerResultsTabSummoner.cpp b/Engine/Plugins/Experimental/AnimNext/Source/AnimNextEditor/Private/Common/AnimNextCompilerResultsTabSummoner.cpp
--- a/Engine/Plugins/Experimental/AnimNext/Source/AnimNextEditor/Private/Common/AnimNextCompilerResultsTabSummoner.cpp
+++ b/Engine/Plugins/Experimental/AnimNext/Source/AnimNextEditor/Private/Common/AnimNextCompilerResultsTabSummoner.cpp
@@ -1,6 +1,10 @@
 // Copyright Epic Games, Inc. All Rights Reserved.
 
 #include "Common/AnimNextCompilerResultsTabSummoner.h"
+
+#include "AnimNextEdGraphNode.h"
+#include "AnimNextRigVMAsset.h"
+#include "EdGraphToken.h"
 #include "IAnimNextEditorModule.h"
 #include "IWorkspaceEditor.h"
 #include "MessageLogModule.h"
@@ -15,6 +19,7 @@ namespace UE::AnimNext::Editor
 
 void SAnimNextCompilerResultsWidget::Construct(const FArguments& InArgs, const TWeakPtr<UE::Workspace::IWorkspaceEditor>& InWorkspaceEditorWeak)
 {
+	WorkspaceEditorWeak = InWorkspaceEditorWeak;
 	CreateMessageLog(InWorkspaceEditorWeak);
 
 	ChildSlot
@@ -38,6 +43,37 @@ void SAnimNextCompilerResultsWidget::CreateMessageLog(const TWeakPtr<UE::Workspa
 		check(MessageLogModule.IsRegisteredLogListing(LogListingName));
 		CompilerResultsListing = MessageLogModule.GetLogListing(LogListingName);
 		CompilerResults = MessageLogModule.CreateLogListingWidget(CompilerResultsListing.ToSharedRef());
+		CompilerResultsListing->OnMessageTokenClicked().AddSP(this, &SAnimNextCompilerResultsWidget::HandleMessageTokenClicked);
+	}
+}
+
+void SAnimNextCompilerResultsWidget::HandleMessageTokenClicked(const TSharedRef<IMessageToken>& InToken)
+{
+	switch(InToken->GetType())
+	{
+	case EMessageToken::EdGraph:
+		{
+			TSharedRef<FEdGraphToken> EdGraphToken = StaticCastSharedRef<FEdGraphToken>(InToken);
+			if(const UAnimNextEdGraphNode* EdGraphNode = Cast<UAnimNextEdGraphNode>(EdGraphToken->GetGraphObject()))
+			{
+				UAnimNextRigVMAsset* Asset = EdGraphNode->GetTypedOuter<UAnimNextRigVMAsset>();
+				if(Asset == nullptr)
+				{
+					return;
+				}
+
+				TSharedPtr<UE::Workspace::IWorkspaceEditor> WorkspaceEditor = WorkspaceEditorWeak.Pin();
+				if(!WorkspaceEditor.IsValid())
+				{
+					return;
+				}
+
+				WorkspaceEditor->OpenObjects({const_cast<UAnimNextEdGraphNode*>(EdGraphNode)});
+			}
+			break;
+		}
+	default:
+		break;
 	}
 }
 
diff --git a/Engine/Plugins/Experimental/AnimNext/Source/AnimNextEditor/Private/Common/AnimNextCompilerResultsTabSummoner.h b/Engine/Plugins/Experimental/AnimNext/Source/AnimNextEditor/Private/Common/AnimNextCompilerResultsTabSummoner.h
--- a/Engine/Plugins/Experimental/AnimNext/Source/AnimNextEditor/Private/Common/AnimNextCompilerResultsTabSummoner.h
+++ b/Engine/Plugins/Experimental/AnimNext/Source/AnimNextEditor/Private/Common/AnimNextCompilerResultsTabSummoner.h
@@ -29,8 +29,11 @@ public:
 private:
 	void CreateMessageLog(const TWeakPtr<UE::Workspace::IWorkspaceEditor>& InWorkspaceEditorWeak);
 
+	void HandleMessageTokenClicked(const TSharedRef<IMessageToken>& InToken);
+
 	TSharedPtr<class IMessageLogListing> CompilerResultsListing;
 	TSharedPtr<class SWidget> CompilerResults;
+	TWeakPtr<UE::Workspace::IWorkspaceEditor> WorkspaceEditorWeak;
 };
 
 struct FAnimNextCompilerResultsTabSummoner : public FWorkflowTabFactory
diff --git a/Engine/Plugins/Experimental/AnimNext/Source/AnimNextEditor/Private/RigVMCommands.cpp b/Engine/Plugins/Experimental/AnimNext/Source/AnimNextEditor/Private/RigVMCommands.cpp
new file mode 100644
--- /dev/null
+++ b/Engine/Plugins/Experimental/AnimNext/Source/AnimNextEditor/Private/RigVMCommands.cpp
@@ -0,0 +1,25 @@
+﻿// Copyright Epic Games, Inc. All Rights Reserved.
+
+#include "RigVMCommands.h"
+
+#define LOCTEXT_NAMESPACE "AnimNextRigVMCommands"
+
+namespace UE::AnimNext
+{
+
+FRigVMCommands::FRigVMCommands()
+	: TCommands<FRigVMCommands>("AnimNextRigVM", LOCTEXT("AnimNextRigVMCommands", "RigVM"), NAME_None, "AnimNextStyle")
+{
+}
+
+void FRigVMCommands::RegisterCommands()
+{
+	UI_COMMAND(Compile, "Compile", "Compile all relevant assets", EUserInterfaceActionType::Button, FInputChord());
+	UI_COMMAND(AutoCompile, "Auto Compile", "Automatically compile on every edit", EUserInterfaceActionType::ToggleButton, FInputChord());
+	UI_COMMAND(CompileWholeWorkspace, "Compile Whole Workspace", "When manually compiling, whether to compile the current file or the whole workspace", EUserInterfaceActionType::ToggleButton, FInputChord());
+	UI_COMMAND(CompileDirtyFiles, "Compile Dirty Files", "When manually compiling, whether to compile only dirty files or files with errors. Files that are not dirty and do not have errors will get skipped.", EUserInterfaceActionType::ToggleButton, FInputChord());
+}
+
+}
+
+#undef LOCTEXT_NAMESPACE
diff --git a/Engine/Plugins/Experimental/AnimNext/Source/AnimNextEditor/Private/RigVMCommands.h b/Engine/Plugins/Experimental/AnimNext/Source/AnimNextEditor/Private/RigVMCommands.h
new file mode 100644
--- /dev/null
+++ b/Engine/Plugins/Experimental/AnimNext/Source/AnimNextEditor/Private/RigVMCommands.h
@@ -0,0 +1,24 @@
+﻿// Copyright Epic Games, Inc. All Rights Reserved.
+
+#pragma once
+
+#include "Framework/Commands/Commands.h"
+
+namespace UE::AnimNext
+{
+
+// The commands for general AnimNext RVM graph editing
+class FRigVMCommands : public TCommands<FRigVMCommands>
+{
+public:
+	FRigVMCommands();
+
+	virtual void RegisterCommands() override;
+
+	TSharedPtr<FUICommandInfo> Compile;
+	TSharedPtr<FUICommandInfo> AutoCompile;
+	TSharedPtr<FUICommandInfo> CompileWholeWorkspace;
+	TSharedPtr<FUICommandInfo> CompileDirtyFiles;
+};
+
+}
diff --git a/Engine/Plugins/Experimental/AnimNext/Source/AnimNextEditor/Private/Workspace/AnimNextWorkspaceEditorMode.cpp b/Engine/Plugins/Experimental/AnimNext/Source/AnimNextEditor/Private/Workspace/AnimNextWorkspaceEditorMode.cpp
new file mode 100644
--- /dev/null
+++ b/Engine/Plugins/Experimental/AnimNext/Source/AnimNextEditor/Private/Workspace/AnimNextWorkspaceEditorMode.cpp
@@ -0,0 +1,383 @@
+﻿// Copyright Epic Games, Inc. All Rights Reserved.
+
+#include "AnimNextWorkspaceEditorMode.h"
+
+#include "AnimNextEditorContext.h"
+#include "AnimNextEditorModule.h"
+#include "AnimNextRigVMAsset.h"
+#include "AnimNextRigVMAssetEditorData.h"
+#include "AnimNextScopedCompilerResults.h"
+#include "AnimNextWorkspaceSchema.h"
+#include "ContextObjectStore.h"
+#include "IAssetCompilationHandler.h"
+#include "InteractiveToolManager.h"
+#include "IWorkspaceEditor.h"
+#include "RigVMCommands.h"
+#include "UncookedOnlyUtils.h"
+#include "Logging/MessageLog.h"
+#include "Modules/ModuleManager.h"
+#include "Toolkits/BaseToolkit.h"
+
+#define LOCTEXT_NAMESPACE "AnimNextWorkspaceEditorMode"
+
+const FEditorModeID UAnimNextWorkspaceEditorMode::EM_AnimNextWorkspace("AnimNextWorkspace");
+
+UAnimNextWorkspaceEditorMode::UAnimNextWorkspaceEditorMode()
+{
+	Info = FEditorModeInfo(UAnimNextWorkspaceEditorMode::EM_AnimNextWorkspace,
+		LOCTEXT("AnimNextWorkspaceEditorModeName", "AnimNextWorkspaceEditorMode"),
+		FSlateIcon(),
+		false);
+}
+
+void UAnimNextWorkspaceEditorMode::Enter()
+{
+	Super::Enter();
+
+	TSharedPtr<UE::Workspace::IWorkspaceEditor> WorkspaceEditor = GetWorkspaceEditor();
+	if(!WorkspaceEditor.IsValid())
+	{
+		return;
+	}
+	
+	WorkspaceEditor->OnFocussedDocumentChanged().AddUObject(this, &UAnimNextWorkspaceEditorMode::HandleFocussedDocumentChanged);
+
+	TArray<UObject*> Assets;
+	WorkspaceEditor->GetOpenedAssets<UAnimNextRigVMAsset>(Assets);
+	for(UObject* Asset : Assets)
+	{
+		UAnimNextRigVMAsset* AnimNextRigVMAsset = static_cast<UAnimNextRigVMAsset*>(Asset);
+		UAnimNextRigVMAssetEditorData* EditorData = UE::AnimNext::UncookedOnly::FUtils::GetEditorData<UAnimNextRigVMAssetEditorData>(AnimNextRigVMAsset);
+
+		EditorData->RigVMCompiledEvent.RemoveAll(this);
+		EditorData->RigVMCompiledEvent.AddUObject(this, &UAnimNextWorkspaceEditorMode::HandleRigVMCompiledEvent);
+		EditorData->RigVMGraphModifiedEvent.RemoveAll(this);
+		EditorData->RigVMGraphModifiedEvent.AddUObject(this, &UAnimNextWorkspaceEditorMode::HandleRigVMModifiedEvent);
+
+		WeakAssets.Add(Asset);
+	}
+
+	UpdateCompileStatus();
+}
+
+void UAnimNextWorkspaceEditorMode::Exit()
+{
+	Super::Exit();
+
+	for(TObjectKey<UObject> Asset : WeakAssets)
+	{
+		UAnimNextRigVMAsset* AnimNextRigVMAsset = Cast<UAnimNextRigVMAsset>(Asset.ResolveObjectPtr());
+		if(AnimNextRigVMAsset == nullptr)
+		{
+			continue;;
+		}
+		
+		UAnimNextRigVMAssetEditorData* EditorData = UE::AnimNext::UncookedOnly::FUtils::GetEditorData<UAnimNextRigVMAssetEditorData>(AnimNextRigVMAsset);
+
+		EditorData->RigVMCompiledEvent.RemoveAll(this);
+		EditorData->RigVMGraphModifiedEvent.RemoveAll(this);
+	}
+
+	WeakAssets.Reset();
+}
+
+void UAnimNextWorkspaceEditorMode::BindCommands()
+{
+	Super::BindCommands();
+
+	TSharedPtr<UE::Workspace::IWorkspaceEditor> WorkspaceEditor = GetWorkspaceEditor();
+	if(!WorkspaceEditor.IsValid())
+	{
+		return;
+	}
+
+	TSharedRef<FUICommandList> ToolkitCommands = WorkspaceEditor->GetToolkitCommands();
+
+	const UE::AnimNext::FRigVMCommands& RigVMCommands = UE::AnimNext::FRigVMCommands::Get();
+	ToolkitCommands->MapAction(RigVMCommands.Compile,
+		FExecuteAction::CreateUObject(this, &UAnimNextWorkspaceEditorMode::HandleCompile));
+
+	ToolkitCommands->MapAction(RigVMCommands.AutoCompile,
+		FExecuteAction::CreateUObject(this, &UAnimNextWorkspaceEditorMode::HandleAutoCompile),
+		FCanExecuteAction(),
+		FIsActionChecked::CreateUObject(this, &UAnimNextWorkspaceEditorMode::IsAutoCompileChecked));
+
+	ToolkitCommands->MapAction(RigVMCommands.CompileWholeWorkspace,
+		FExecuteAction::CreateUObject(this, &UAnimNextWorkspaceEditorMode::HandleCompileWholeWorkspace),
+		FCanExecuteAction(),
+		FIsActionChecked::CreateUObject(this, &UAnimNextWorkspaceEditorMode::IsCompileWholeWorkspaceChecked));
+
+	ToolkitCommands->MapAction(RigVMCommands.CompileDirtyFiles,
+		FExecuteAction::CreateUObject(this, &UAnimNextWorkspaceEditorMode::HandleCompileDirtyFiles),
+		FCanExecuteAction(),
+		FIsActionChecked::CreateUObject(this, &UAnimNextWorkspaceEditorMode::IsCompileDirtyFilesChecked));
+}
+
+void UAnimNextWorkspaceEditorMode::HandleCompile()
+{
+	TSharedPtr<UE::Workspace::IWorkspaceEditor> WorkspaceEditor = GetWorkspaceEditor();
+	if(!WorkspaceEditor.IsValid())
+	{
+		return;
+	}
+
+	TArray<UObject*> Assets;
+
+	UObject* Asset = nullptr;
+	FText JobName = LOCTEXT("DefaultJobName", "Job");
+	if(IsCompileWholeWorkspaceChecked())
+	{
+		WorkspaceEditor->GetOpenedAssets<UObject>(Assets);
+		if(Assets.Num() == 1)
+		{
+			Asset = Assets[0];
+			JobName = FText::FromName(Asset->GetFName());
+		}
+		else
+		{
+			Asset = WorkspaceEditor->GetWorkspaceAsset();
+			JobName = FText::FromName(Asset->GetFName());
+		}
+	}
+	else if(UObject* FocussedDocument = WorkspaceEditor->GetFocussedDocument())
+	{
+		// Find the outer asset for this document
+		Asset = FocussedDocument;
+		
+		while(Asset && (!Asset->IsAsset() || Asset->IsPackageExternal()))
+		{
+			Asset = Asset->GetOuter();
+		}
+
+		if(Asset)
+		{
+			JobName = FText::FromName(Asset->GetFName());
+			Assets.Add(Asset);
+		}
+	}
+
+	if(IsCompileDirtyFilesChecked())
+	{
+		for(int32 AssetIndex = 0; AssetIndex < Assets.Num(); ++AssetIndex)
+		{
+			UObject* AssetToCheck = Assets[AssetIndex];
+			if(TSharedPtr<UE::AnimNext::Editor::IAssetCompilationHandler> FoundCompiler = GetAssetCompiler(AssetToCheck))
+			{
+				UE::AnimNext::Editor::ECompileStatus AssetCompileStatus = FoundCompiler->GetCompileStatus(WorkspaceEditor.ToSharedRef(), AssetToCheck);
+				if( AssetCompileStatus != UE::AnimNext::Editor::ECompileStatus::Dirty &&
+					AssetCompileStatus != UE::AnimNext::Editor::ECompileStatus::Error)
+				{
+					Assets.RemoveAtSwap(AssetIndex);
+					AssetIndex--;
+				}
+			}
+		}
+	}
+
+	if(Assets.Num() == 0)
+	{
+		return;
+	}
+
+	// Start a batch compile scope
+	UE::AnimNext::UncookedOnly::FScopedCompilerResults CompileResults(JobName, Asset, Assets);
+
+	UE::AnimNext::Editor::FAnimNextEditorModule& AnimNextEditorModule = FModuleManager::LoadModuleChecked<UE::AnimNext::Editor::FAnimNextEditorModule>("AnimNextEditor");
+	for(UObject* AssetToCompile : Assets)
+	{
+		if(TSharedPtr<UE::AnimNext::Editor::IAssetCompilationHandler> FoundCompiler = GetAssetCompiler(AssetToCompile))
+		{
+			FoundCompiler->Compile(WorkspaceEditor.ToSharedRef(), AssetToCompile);
+		}
+	}
+}
+
+void UAnimNextWorkspaceEditorMode::HandleAutoCompile()
+{
+	State.bAutoCompile = !State.bAutoCompile;
+
+	TSharedPtr<UE::Workspace::IWorkspaceEditor> WorkspaceEditor = GetWorkspaceEditor();
+	if(WorkspaceEditor.IsValid())
+	{
+		PropagateAutoCompile(WorkspaceEditor.ToSharedRef(), State.bAutoCompile);
+	}
+}
+
+void UAnimNextWorkspaceEditorMode::PropagateAutoCompile(TSharedRef<UE::Workspace::IWorkspaceEditor> InWorkspaceEditor, bool bInAutoCompile)
+{
+	TArray<UObject*> Assets;
+	InWorkspaceEditor->GetOpenedAssets(Assets);
+	UE::AnimNext::Editor::FAnimNextEditorModule& AnimNextEditorModule = FModuleManager::LoadModuleChecked<UE::AnimNext::Editor::FAnimNextEditorModule>("AnimNextEditor");
+	for(UObject* Asset : Assets)
+	{
+		if(TSharedPtr<UE::AnimNext::Editor::IAssetCompilationHandler> FoundCompiler = GetAssetCompiler(Asset))
+		{
+			FoundCompiler->SetAutoCompile(InWorkspaceEditor, Asset, bInAutoCompile);
+		}
+	}
+}
+
+bool UAnimNextWorkspaceEditorMode::IsAutoCompileChecked() const
+{
+	return State.bAutoCompile;
+}
+
+void UAnimNextWorkspaceEditorMode::HandleCompileWholeWorkspace()
+{
+	State.bCompileWholeWorkspace = !State.bCompileWholeWorkspace;
+}
+
+bool UAnimNextWorkspaceEditorMode::IsCompileWholeWorkspaceChecked() const
+{
+	return State.bCompileWholeWorkspace;
+}
+
+void UAnimNextWorkspaceEditorMode::HandleCompileDirtyFiles()
+{
+	State.bCompileDirtyFiles = !State.bCompileDirtyFiles;
+}
+
+bool UAnimNextWorkspaceEditorMode::IsCompileDirtyFilesChecked() const
+{
+	return State.bCompileDirtyFiles;
+}
+
+void UAnimNextWorkspaceEditorMode::HandleFocussedDocumentChanged(TObjectPtr<UObject> InObject)
+{
+	if(InObject == nullptr)
+	{
+		return;
+	}
+
+	UAnimNextRigVMAsset* AnimNextRigVMAsset = Cast<UAnimNextRigVMAsset>(InObject);
+	if(AnimNextRigVMAsset == nullptr)
+	{
+		AnimNextRigVMAsset = InObject->GetTypedOuter<UAnimNextRigVMAsset>();
+	}
+	
+	if(AnimNextRigVMAsset == nullptr)
+	{
+		return;
+	}
+
+	// Ensure auto-compilation is propagated for newly opened asset 
+	UAnimNextRigVMAssetEditorData* EditorData = UE::AnimNext::UncookedOnly::FUtils::GetEditorData<UAnimNextRigVMAssetEditorData>(AnimNextRigVMAsset);
+	EditorData->SetAutoVMRecompile(State.bAutoCompile);
+
+	// Subscribe to asset compilation/modification
+	EditorData->RigVMCompiledEvent.RemoveAll(this);
+	EditorData->RigVMCompiledEvent.AddUObject(this, &UAnimNextWorkspaceEditorMode::HandleRigVMCompiledEvent);
+	EditorData->RigVMGraphModifiedEvent.RemoveAll(this);
+	EditorData->RigVMGraphModifiedEvent.AddUObject(this, &UAnimNextWorkspaceEditorMode::HandleRigVMModifiedEvent);
+
+	WeakAssets.Add(AnimNextRigVMAsset);
+}
+
+void UAnimNextWorkspaceEditorMode::UpdateCompileStatus()
+{
+	bUpdateCompileStatus = false;
+
+	TSharedPtr<UE::Workspace::IWorkspaceEditor> WorkspaceEditor = GetWorkspaceEditor();
+	if(!WorkspaceEditor.IsValid())
+	{
+		return;
+	}
+
+	int32 NumAssetsWithWarnings = 0;
+	int32 NumAssetsWithErrors = 0;
+	int32 NumAssetsDirty = 0;
+
+	TArray<UObject*> Assets;
+	WorkspaceEditor->GetOpenedAssets<UAnimNextRigVMAsset>(Assets);
+
+	for(UObject* Asset : Assets)
+	{
+		if(TSharedPtr<UE::AnimNext::Editor::IAssetCompilationHandler> FoundCompiler = GetAssetCompiler(Asset))
+		{
+			UE::AnimNext::Editor::ECompileStatus AssetCompileStatus = FoundCompiler->GetCompileStatus(WorkspaceEditor.ToSharedRef(), Asset);
+			if(AssetCompileStatus == UE::AnimNext::Editor::ECompileStatus::Dirty)
+			{
+				NumAssetsDirty++;
+			}
+
+			if(AssetCompileStatus == UE::AnimNext::Editor::ECompileStatus::Error)
+			{
+				NumAssetsWithErrors++;
+			}
+
+			if(AssetCompileStatus == UE::AnimNext::Editor::ECompileStatus::Warning)
+			{
+				NumAssetsWithWarnings++;
+			}
+		}
+	}
+
+	if(NumAssetsWithErrors > 0)
+	{
+		CompileStatus = UE::AnimNext::Editor::ECompileStatus::Error;
+	}
+	else if(NumAssetsWithWarnings > 0)
+	{
+		CompileStatus = UE::AnimNext::Editor::ECompileStatus::Warning;
+	}
+	else if(NumAssetsDirty > 0)
+	{
+		CompileStatus = UE::AnimNext::Editor::ECompileStatus::Dirty;
+	}
+	else
+	{
+		CompileStatus = UE::AnimNext::Editor::ECompileStatus::UpToDate;
+	}
+}
+
+void UAnimNextWorkspaceEditorMode::HandleRigVMCompiledEvent(UObject* InAsset, URigVM* InVM, FRigVMExtendedExecuteContext& InExtendedExecuteContext)
+{
+	UpdateCompileStatus();
+}
+
+void UAnimNextWorkspaceEditorMode::HandleRigVMModifiedEvent(ERigVMGraphNotifType InType, URigVMGraph* InGraph, UObject* InSubject)
+{
+	if(InGraph == nullptr)
+	{
+		return;
+	}
+
+	UAnimNextRigVMAssetEditorData* EditorData = InGraph->GetTypedOuter<UAnimNextRigVMAssetEditorData>();
+	if(EditorData->IsDirtyForRecompilation())
+	{
+		CompileStatus = UE::AnimNext::Editor::ECompileStatus::Dirty;
+	}
+}
+
+TSharedPtr<UE::Workspace::IWorkspaceEditor> UAnimNextWorkspaceEditorMode::GetWorkspaceEditor()
+{
+	if (const UContextObjectStore* ContextObjectStore = GetToolManager()->GetContextObjectStore())
+	{
+		if (const UAnimNextEditorContext* Context = ContextObjectStore->FindContext<UAnimNextEditorContext>())
+		{
+			return Context->WeakWorkspaceEditor.Pin();
+		}
+	}
+
+	return nullptr;
+}
+
+TSharedPtr<UE::AnimNext::Editor::IAssetCompilationHandler> UAnimNextWorkspaceEditorMode::GetAssetCompiler(UObject* InAsset)
+{
+	if(TSharedRef<UE::AnimNext::Editor::IAssetCompilationHandler>* FoundCompiler = AssetCompilers.Find(InAsset))
+	{
+		return *FoundCompiler;
+	}
+
+	UE::AnimNext::Editor::FAnimNextEditorModule& AnimNextEditorModule = FModuleManager::LoadModuleChecked<UE::AnimNext::Editor::FAnimNextEditorModule>("AnimNextEditor");
+	if(const UE::AnimNext::Editor::FAssetCompilationHandlerFactoryDelegate* FoundFactory = AnimNextEditorModule.FindAssetCompilationHandlerFactory(InAsset->GetClass()))
+	{
+		TSharedRef<UE::AnimNext::Editor::IAssetCompilationHandler> NewCompiler = AssetCompilers.Add(InAsset, FoundFactory->Execute(InAsset));
+		NewCompiler->OnCompileStatusChanged().BindUObject(this, &UAnimNextWorkspaceEditorMode::UpdateCompileStatus);
+		return NewCompiler;
+	}
+
+	return nullptr;
+}
+
+#undef LOCTEXT_NAMESPACE
\ No newline at end of file
diff --git a/Engine/Plugins/Experimental/AnimNext/Source/AnimNextEditor/Private/Workspace/AnimNextWorkspaceEditorMode.h b/Engine/Plugins/Experimental/AnimNext/Source/AnimNextEditor/Private/Workspace/AnimNextWorkspaceEditorMode.h
new file mode 100644
--- /dev/null
+++ b/Engine/Plugins/Experimental/AnimNext/Source/AnimNextEditor/Private/Workspace/AnimNextWorkspaceEditorMode.h
@@ -0,0 +1,89 @@
+﻿// Copyright Epic Games, Inc. All Rights Reserved.
+
+#pragma once
+
+#include "AnimNextWorkspaceState.h"
+#include "EdMode.h"
+#include "IAssetCompilationHandler.h"
+#include "AnimNextWorkspaceEditorMode.generated.h"
+
+class UAnimNextWorkspaceSchema;
+class URigVM;
+class URigVMGraph;
+struct FRigVMExtendedExecuteContext;
+enum class ERigVMGraphNotifType : uint8;
+
+namespace UE::Workspace
+{
+	class IWorkspaceEditor;
+}
+
+UCLASS(Transient)
+class UAnimNextWorkspaceEditorMode : public UEdMode
+{
+	GENERATED_BODY()
+
+public:
+	const static FEditorModeID EM_AnimNextWorkspace;
+
+	UAnimNextWorkspaceEditorMode();
+
+	// Get the current compilation status
+	UE::AnimNext::Editor::ECompileStatus GetLatestCompileStatus() const { return CompileStatus; }
+
+	// Get the currently stored workspace state
+	const FAnimNextWorkspaceState& GetState() { return State; } 
+
+private:
+	// UEdMode interface
+	virtual bool IsCompatibleWith(FEditorModeID OtherModeID) const override { return true; }
+	virtual void Enter() override;
+	virtual void Exit() override;
+	virtual void BindCommands() override;
+	virtual bool UsesToolkits() const override { return false; }
+
+	void HandleCompile();
+	void HandleAutoCompile();
+	bool IsAutoCompileChecked() const;
+	void HandleCompileWholeWorkspace();
+	bool IsCompileWholeWorkspaceChecked() const;
+	void HandleCompileDirtyFiles();
+	bool IsCompileDirtyFilesChecked() const;
+
+	// Updates the compile status. Scans all opened assets in the workspace.
+	void UpdateCompileStatus();
+
+	// Handle an asset in the workspace getting compiled 
+	void HandleRigVMCompiledEvent(UObject* InAsset, URigVM* InVM, FRigVMExtendedExecuteContext& InExtendedExecuteContext);
+
+	// Handle an asset getting modified
+	void HandleRigVMModifiedEvent(ERigVMGraphNotifType InType, URigVMGraph* InGraph, UObject* InSubject);
+
+	// Handle the document focus changing
+	void HandleFocussedDocumentChanged(TObjectPtr<UObject> InAsset); 
+
+	// Forward auto-compilation to all assets in the workspace
+	void PropagateAutoCompile(TSharedRef<UE::Workspace::IWorkspaceEditor> InWorkspaceEditor, bool bInAutoCompile);
+
+	TSharedPtr<UE::Workspace::IWorkspaceEditor> GetWorkspaceEditor();
+
+	// Lazily construct an asset compiler for the supplied asset
+	TSharedPtr<UE::AnimNext::Editor::IAssetCompilationHandler> GetAssetCompiler(UObject* InAsset);
+
+	friend UAnimNextWorkspaceSchema;
+
+	// Custom state, persisted via UAnimNextWorkspaceSchema
+	FAnimNextWorkspaceState State;
+
+	// Asset compilers for all current assets 
+	TMap<FObjectKey, TSharedRef<UE::AnimNext::Editor::IAssetCompilationHandler>> AssetCompilers;
+
+	// All assets that we are currently tracking for compilation status
+	TSet<TObjectKey<UObject>> WeakAssets;
+
+	// Current compilation status
+	UE::AnimNext::Editor::ECompileStatus CompileStatus = UE::AnimNext::Editor::ECompileStatus::Unknown;
+
+	// Flag to update the latest compile status
+	bool bUpdateCompileStatus = false;
+};
\ No newline at end of file
diff --git a/Engine/Plugins/Experimental/AnimNext/Source/AnimNextEditor/Private/Workspace/AnimNextWorkspaceSchema.cpp b/Engine/Plugins/Experimental/AnimNext/Source/AnimNextEditor/Private/Workspace/AnimNextWorkspaceSchema.cpp
--- a/Engine/Plugins/Experimental/AnimNext/Source/AnimNextEditor/Private/Workspace/AnimNextWorkspaceSchema.cpp
+++ b/Engine/Plugins/Experimental/AnimNext/Source/AnimNextEditor/Private/Workspace/AnimNextWorkspaceSchema.cpp
@@ -3,6 +3,10 @@
 #include "AnimNextWorkspaceSchema.h"
 
 #include "AnimNextEditorModule.h"
+#include "AnimNextWorkspaceEditorMode.h"
+#include "AssetEditorModeManager.h"
+#include "IWorkspaceEditor.h"
+#include "StructUtils/InstancedStruct.h"
 
 #define LOCTEXT_NAMESPACE "AnimNextWorkspaceSchema"
 
@@ -18,4 +22,36 @@ TConstArrayView<FTopLevelAssetPath> UAnimNextWorkspaceSchema::GetSupportedAssetC
 	return Module.SupportedAssetClasses;
 }
 
+void UAnimNextWorkspaceSchema::OnSaveWorkspaceState(TSharedRef<UE::Workspace::IWorkspaceEditor> InWorkspaceEditor, FInstancedStruct& OutWorkspaceState) const
+{
+	UAnimNextWorkspaceEditorMode* EditorMode = Cast<UAnimNextWorkspaceEditorMode>(InWorkspaceEditor->GetEditorModeManager().GetActiveScriptableMode(UAnimNextWorkspaceEditorMode::EM_AnimNextWorkspace));
+	if(EditorMode == nullptr)
+	{
+		return;
+	}
+
+	OutWorkspaceState = FInstancedStruct::Make(EditorMode->State);
+}
+
+void UAnimNextWorkspaceSchema::OnLoadWorkspaceState(TSharedRef< UE::Workspace::IWorkspaceEditor> InWorkspaceEditor, const FInstancedStruct& InWorkspaceState) const
+{
+	// Activate and set up our mode
+	InWorkspaceEditor->GetEditorModeManager().ActivateMode(UAnimNextWorkspaceEditorMode::EM_AnimNextWorkspace);
+
+	UAnimNextWorkspaceEditorMode* EditorMode = Cast<UAnimNextWorkspaceEditorMode>(InWorkspaceEditor->GetEditorModeManager().GetActiveScriptableMode(UAnimNextWorkspaceEditorMode::EM_AnimNextWorkspace));
+	if(EditorMode == nullptr)
+	{
+		return;
+	}
+
+	const FAnimNextWorkspaceState* State = InWorkspaceState.GetPtr<FAnimNextWorkspaceState>();
+	if(State == nullptr)
+	{
+		return;
+	}
+
+	EditorMode->State = *State;
+	EditorMode->PropagateAutoCompile(InWorkspaceEditor, EditorMode->State.bAutoCompile);
+}
+
 #undef LOCTEXT_NAMESPACE
\ No newline at end of file
diff --git a/Engine/Plugins/Experimental/AnimNext/Source/AnimNextEditor/Private/Workspace/AnimNextWorkspaceSchema.h b/Engine/Plugins/Experimental/AnimNext/Source/AnimNextEditor/Private/Workspace/AnimNextWorkspaceSchema.h
--- a/Engine/Plugins/Experimental/AnimNext/Source/AnimNextEditor/Private/Workspace/AnimNextWorkspaceSchema.h
+++ b/Engine/Plugins/Experimental/AnimNext/Source/AnimNextEditor/Private/Workspace/AnimNextWorkspaceSchema.h
@@ -14,4 +14,6 @@ class UAnimNextWorkspaceSchema : public UWorkspaceSchema
 	// UWorkspaceSchema interface
 	virtual FText GetDisplayName() const override;
 	virtual TConstArrayView<FTopLevelAssetPath> GetSupportedAssetClassPaths() const override;
+	virtual void OnSaveWorkspaceState(TSharedRef<UE::Workspace::IWorkspaceEditor> InWorkspaceEditor, FInstancedStruct& OutWorkspaceState) const override;
+	virtual void OnLoadWorkspaceState(TSharedRef<UE::Workspace::IWorkspaceEditor> InWorkspaceEditor, const FInstancedStruct& InWorkspaceState) const override;
 };
\ No newline at end of file
diff --git a/Engine/Plugins/Experimental/AnimNext/Source/AnimNextEditor/Private/Workspace/AnimNextWorkspaceState.h b/Engine/Plugins/Experimental/AnimNext/Source/AnimNextEditor/Private/Workspace/AnimNextWorkspaceState.h
new file mode 100644
--- /dev/null
+++ b/Engine/Plugins/Experimental/AnimNext/Source/AnimNextEditor/Private/Workspace/AnimNextWorkspaceState.h
@@ -0,0 +1,24 @@
+﻿// Copyright Epic Games, Inc. All Rights Reserved.
+
+#pragma once
+
+#include "AnimNextWorkspaceState.generated.h"
+
+// Custom persistent workspace state 
+USTRUCT()
+struct FAnimNextWorkspaceState
+{
+	GENERATED_BODY()
+
+	// Whether to automatically compile RigVM files after every edit
+	UPROPERTY()
+	bool bAutoCompile = true;
+
+	// When manually compiling, whether to compile the current file or the whole workspace
+	UPROPERTY()
+	bool bCompileWholeWorkspace = true;
+
+	// When manually compiling, whether to compile only dirty files. Files that are not dirty will get skipped.
+	UPROPERTY()
+	bool bCompileDirtyFiles = true;
+};
\ No newline at end of file
diff --git a/Engine/Plugins/Experimental/AnimNext/Source/AnimNextEditor/Public/AssetCompilationHandler.h b/Engine/Plugins/Experimental/AnimNext/Source/AnimNextEditor/Public/AssetCompilationHandler.h
new file mode 100644
--- /dev/null
+++ b/Engine/Plugins/Experimental/AnimNext/Source/AnimNextEditor/Public/AssetCompilationHandler.h
@@ -0,0 +1,28 @@
+﻿// Copyright Epic Games, Inc. All Rights Reserved.
+
+#pragma once
+
+#include "CoreMinimal.h"
+#include "IAssetCompilationHandler.h"
+
+namespace UE::AnimNext::Editor
+{
+
+// Asset compiler that can compile assets based on UAnimNextRigVMAsset
+class ANIMNEXTEDITOR_API FAssetCompilationHandler : public IAssetCompilationHandler
+{
+public:
+	explicit FAssetCompilationHandler(UObject* InAsset);
+
+protected:
+	// IAssetCompilationHandler interface
+	virtual void Compile(TSharedRef<Workspace::IWorkspaceEditor> InWorkspaceEditor, UObject* InAsset) override;
+	virtual void SetAutoCompile(TSharedRef<Workspace::IWorkspaceEditor> InWorkspaceEditor, UObject* InAsset, bool bInAutoCompile) override;
+	virtual bool GetAutoCompile(TSharedRef<Workspace::IWorkspaceEditor> InWorkspaceEditor, const UObject* InAsset) const override;
+	virtual ECompileStatus GetCompileStatus(TSharedRef<Workspace::IWorkspaceEditor> InWorkspaceEditor, const UObject* InAsset) const override;
+
+	int32 NumErrors = 0;
+	int32 NumWarnings = 0;
+};
+
+}
diff --git a/Engine/Plugins/Experimental/AnimNext/Source/AnimNextEditor/Public/IAnimNextEditorModule.h b/Engine/Plugins/Experimental/AnimNext/Source/AnimNextEditor/Public/IAnimNextEditorModule.h
--- a/Engine/Plugins/Experimental/AnimNext/Source/AnimNextEditor/Public/IAnimNextEditorModule.h
+++ b/Engine/Plugins/Experimental/AnimNext/Source/AnimNextEditor/Public/IAnimNextEditorModule.h
@@ -4,6 +4,7 @@
 
 #include "CoreMinimal.h"
 #include "Modules/ModuleManager.h"
+#include "UObject/TopLevelAssetPath.h"
 
 class SWidget;
 class IMessageLogListing;
@@ -23,6 +24,10 @@ const FLazyName LogListingName("AnimNextCompilerResults");
 
 struct FVariablePickerArgs;
 struct FActionMenuContextData;
+class IAssetCompilationHandler;
+
+// A factory function used to make an asset compilation handler for an asset
+using FAssetCompilationHandlerFactoryDelegate = TDelegate<TSharedRef<IAssetCompilationHandler>(UObject*)>;
 
 class IAnimNextEditorModule : public IModuleInterface
 {
@@ -54,6 +59,15 @@ public:
 	// Unregister a graph context menu actions provider
 	// @param InDelegateHandle Handle of the delegate to remove
 	virtual void UnregisterGraphMenuActionsProvider(const FDelegateHandle& InDelegateHandle) = 0;
+
+	// Register an asset compilation handler
+	// @param InClassPath The path of the asset's class
+	// @param InAssetCompilerFactory A factory function used to make an asset compiler for an asset
+	virtual void RegisterAssetCompilationHandler(const FTopLevelAssetPath& InClassPath, FAssetCompilationHandlerFactoryDelegate InAssetCompilationHandlerFactory) = 0;
+
+	// Unregister an asset compilation handler
+	// @param InClassPath The path of the asset's class
+	virtual void UnregisterAssetCompilationHandler(const FTopLevelAssetPath& InClassPath) = 0;
 };
 
 }
diff --git a/Engine/Plugins/Experimental/AnimNext/Source/AnimNextEditor/Public/IAssetCompilationHandler.h b/Engine/Plugins/Experimental/AnimNext/Source/AnimNextEditor/Public/IAssetCompilationHandler.h
new file mode 100644
--- /dev/null
+++ b/Engine/Plugins/Experimental/AnimNext/Source/AnimNextEditor/Public/IAssetCompilationHandler.h
@@ -0,0 +1,49 @@
+﻿// Copyright Epic Games, Inc. All Rights Reserved.
+
+#pragma once
+
+#include "CoreMinimal.h"
+
+namespace UE::Workspace
+{
+class IWorkspaceEditor;
+}
+
+namespace UE::AnimNext::Editor
+{
+
+// Status of compilation. Order matters - status is compared by magnitude to indicate severity
+enum class ECompileStatus : uint8
+{
+	Unknown,
+	UpToDate,
+	Dirty,
+	Warning,
+	Error,
+};
+
+class IAssetCompilationHandler : public TSharedFromThis<IAssetCompilationHandler>
+{
+public:
+	virtual ~IAssetCompilationHandler() = default;
+
+	// Called to compile an asset from the editor
+	virtual void Compile(TSharedRef<Workspace::IWorkspaceEditor> InWorkspaceEditor, UObject* InAsset) = 0;
+
+	// Called to set the auto compile mode of the asset
+	virtual void SetAutoCompile(TSharedRef<Workspace::IWorkspaceEditor> InWorkspaceEditor, UObject* InAsset, bool bInAutoCompile) {}
+
+	// Called to get the auto compile mode of the asset
+	virtual bool GetAutoCompile(TSharedRef<Workspace::IWorkspaceEditor> InWorkspaceEditor, const UObject* InAsset) const { return false; }
+
+	// Called to get the compile status of the asset
+	virtual ECompileStatus GetCompileStatus(TSharedRef<Workspace::IWorkspaceEditor> InWorkspaceEditor, const UObject* InAsset) const { return ECompileStatus::Unknown; }
+
+	// Called to notify that the compile status of the asset may have changed
+	FSimpleDelegate& OnCompileStatusChanged() { return CompileStatusChangedDelegate; }
+
+private:
+	FSimpleDelegate CompileStatusChangedDelegate;
+};
+
+}
diff --git a/Engine/Plugins/Experimental/AnimNext/Source/AnimNextUncookedOnly/Internal/AnimNextRigVMAssetEditorData.h b/Engine/Plugins/Experimental/AnimNext/Source/AnimNextUncookedOnly/Internal/AnimNextRigVMAssetEditorData.h
--- a/Engine/Plugins/Experimental/AnimNext/Source/AnimNextUncookedOnly/Internal/AnimNextRigVMAssetEditorData.h
+++ b/Engine/Plugins/Experimental/AnimNext/Source/AnimNextUncookedOnly/Internal/AnimNextRigVMAssetEditorData.h
@@ -15,11 +15,13 @@ enum class ERigVMGraphNotifType : uint8;
 class UAnimNextRigVMAssetEditorData;
 class UAnimNextEdGraph;
 class UAnimNextDataInterfaceEntry;
+class UAnimNextWorkspaceEditorMode;
 
 namespace UE::AnimNext::UncookedOnly
 {
 	struct FUtils;
 	struct FUtilsPrivate;
+	class FScopedCompilerResults;
 }
 
 namespace UE::AnimNext::Editor
@@ -44,6 +46,7 @@ namespace UE::AnimNext::Editor
 	struct FPublicVariablesImpl;
 	class FAnimNextAnimGraphEditorModule;
 	class FCallFunctionSharedDataDetails;
+	class FAssetCompilationHandler;
 }
 
 namespace UE::AnimNext::Tests
@@ -184,6 +187,9 @@ protected:
 	friend struct UE::AnimNext::Editor::FPublicVariablesImpl;
 	friend UE::AnimNext::Editor::FAnimNextAnimGraphEditorModule;
 	friend class UE::AnimNext::Editor::FCallFunctionSharedDataDetails;
+	friend UAnimNextWorkspaceEditorMode;
+	friend UE::AnimNext::UncookedOnly::FScopedCompilerResults;
+	friend UE::AnimNext::Editor::FAssetCompilationHandler;
 
 	// UObject interface
 	virtual void Serialize(FArchive& Ar) override;
@@ -433,6 +439,9 @@ protected:
 	// Refresh the 'external' models for the RigVM client to reference
 	void RefreshExternalModels();
 
+	// Clear the error info for all EdGraphNodes
+	void ClearErrorInfoForAllEdGraphs();
+
 	// Handle compiler reporting
 	void HandleReportFromCompiler(EMessageSeverity::Type InSeverity, UObject* InSubject, const FString& InMessage);
 
@@ -451,6 +460,9 @@ protected:
 	// Handle replacing a notify
 	static void HandleReplaceNotify(UObject* InAsset, const FString& InFindString, const FString& InReplaceString, bool bFindWholeWord, ESearchCase::Type InSearchCase);
 
+	// Check whether this asset should be recompiled
+	bool IsDirtyForRecompilation() const;
+
 	/** All entries in this asset - not saved, discovered at load time and also contains InternalEntries */
 	UPROPERTY(transient)
 	TArray<TObjectPtr<UAnimNextRigVMAssetEntry>> Entries;
@@ -522,6 +534,7 @@ protected:
 
 	bool bAutoRecompileVM = true;
 	bool bErrorsDuringCompilation = false;
+	bool bWarningsDuringCompilation = false;
 	bool bSuspendModelNotificationsForSelf = false;
 	bool bSuspendAllNotifications = false;
 	bool bCompileInDebugMode = false;
diff --git a/Engine/Plugins/Experimental/AnimNext/Source/AnimNextUncookedOnly/Private/AnimNextRigVMAssetEditorData.cpp b/Engine/Plugins/Experimental/AnimNext/Source/AnimNextUncookedOnly/Private/AnimNextRigVMAssetEditorData.cpp
--- a/Engine/Plugins/Experimental/AnimNext/Source/AnimNextUncookedOnly/Private/AnimNextRigVMAssetEditorData.cpp
+++ b/Engine/Plugins/Experimental/AnimNext/Source/AnimNextUncookedOnly/Private/AnimNextRigVMAssetEditorData.cpp
@@ -7,6 +7,7 @@
 #include "Entries/AnimNextRigVMAssetEntry.h"
 #include "AnimNextRigVMAssetSchema.h"
 #include "AnimNextAssetWorkspaceAssetUserData.h"
+#include "AnimNextScopedCompilerResults.h"
 #include "ControlRigDefines.h"
 #include "ExternalPackageHelper.h"
 #include "IAnimNextRigVMGraphInterface.h"
@@ -14,6 +15,7 @@
 #include "IWorkspaceEditorModule.h"
 #include "ObjectTools.h"
 #include "UncookedOnlyUtils.h"
+#include "Animation/AnimCompressionTypes.h"
 #include "Animation/Skeleton.h"
 #include "DataInterface/AnimNextDataInterface_EditorData.h"
 #include "Entries/AnimNextDataInterfaceEntry.h"
@@ -30,6 +32,7 @@
 #include "UObject/AssetRegistryTagsContext.h"
 #include "DataInterface/AnimNextDataInterface_EditorData.h"
 #include "RigVMFunctions/Execution/RigVMFunction_UserDefinedEvent.h"
+#include "Misc/UObjectToken.h"
 #include "UObject/SavePackage.h"
 
 #if WITH_EDITOR
@@ -41,7 +44,7 @@
 
 void UAnimNextRigVMAssetEditorData::BroadcastModified(EAnimNextEditorDataNotifType InType, UObject* InSubject)
 {
-	RecompileVM();
+	RequestAutoVMRecompilation();
 
 	if(!bSuspendEditorDataNotifications)
 	{
@@ -552,7 +555,7 @@ void UAnimNextRigVMAssetEditorData::HandleRigVMGraphRemoved(const FRigVMClient*
 		GraphModels.Remove(RigVMGraph);
 
 		RemoveEdGraph(RigVMGraph);
-		RecompileVM();
+		RequestAutoVMRecompilation();
 
 #if WITH_EDITOR
 		if(!bSuspendPythonMessagesForRigVMClient)
@@ -759,9 +762,11 @@ void UAnimNextRigVMAssetEditorData::RecompileVM()
 	TGuardValue<bool> CompilingGuard(bIsCompiling, true);
 
 	UAnimNextRigVMAsset* Asset = FUtils::GetAsset<UAnimNextRigVMAsset>(this);
+	FScopedCompilerResults CompilerResults(Asset);
 
 	CachedExports.Reset();  // asset variables and other tags will be updated at the end by AssetRegistry->AssetUpdateTags
 
+	bWarningsDuringCompilation = false;
 	bErrorsDuringCompilation = false;
 
 	RigGraphDisplaySettings.MinMicroSeconds = RigGraphDisplaySettings.LastMinMicroSeconds = DBL_MAX;
@@ -774,12 +779,8 @@ void UAnimNextRigVMAssetEditorData::RecompileVM()
 
 		VMCompileSettings.SetExecuteContextStruct(FAnimNextExecuteContext::StaticStruct());
 		FRigVMCompileSettings Settings = (bCompileInDebugMode) ? FRigVMCompileSettings::Fast(VMCompileSettings.GetExecuteContextStruct()) : VMCompileSettings;
-
-		FMessageLog("AnimNextCompilerResults").NewPage(FText::FromName(Asset->GetFName()));
-		Settings.ASTSettings.ReportDelegate.BindLambda([](EMessageSeverity::Type InType, UObject* InObject, const FString& InString)
-		{
-			FMessageLog("AnimNextCompilerResults").Message(InType, FText::FromString(InString));
-		});
+		Settings.SurpressInfoMessages = false;
+		Settings.ASTSettings.ReportDelegate.BindUObject(this, &UAnimNextRigVMAssetEditorData::HandleReportFromCompiler);
 
 		FUtils::RecreateVM(Asset);
 
@@ -798,15 +799,13 @@ void UAnimNextRigVMAssetEditorData::RecompileVM()
 		TArray<URigVMGraph*> AllGraphs = VMClient->GetAllModels(false, false);
 		AllGraphs.Append(ProgrammaticGraphs);
 
-		if(AllGraphs.Num() == 0)
+		if(AllGraphs.Num() > 0)
 		{
-			return;
-		}
-
-		URigVMController* Controller = VMClient->GetOrCreateController(AllGraphs[0]);
+			URigVMController* Controller = VMClient->GetOrCreateController(AllGraphs[0]);
 
-		URigVMCompiler* Compiler = URigVMCompiler::StaticClass()->GetDefaultObject<URigVMCompiler>();
-		Compiler->Compile(Settings, AllGraphs, Controller, Asset->VM, Asset->ExtendedExecuteContext, Asset->GetExternalVariables(), &PinToOperandMap);
+			URigVMCompiler* Compiler = URigVMCompiler::StaticClass()->GetDefaultObject<URigVMCompiler>();
+			Compiler->Compile(Settings, AllGraphs, Controller, Asset->VM, Asset->ExtendedExecuteContext, Asset->GetExternalVariables(), &PinToOperandMap);
+		}
 
 		// Initialize right away, in packaged builds we initialize during PostLoad
 		Asset->VM->Initialize(Asset->ExtendedExecuteContext);
@@ -824,6 +823,7 @@ void UAnimNextRigVMAssetEditorData::RecompileVM()
 		}
 
 		bVMRecompilationRequired = false;
+
 		if(Asset->VM)
 		{
 			RigVMCompiledEvent.Broadcast(Asset, Asset->VM, Asset->ExtendedExecuteContext);
@@ -831,28 +831,28 @@ void UAnimNextRigVMAssetEditorData::RecompileVM()
 
 		FAnimNextAssetRegistryExports Exports;
 		FUtils::GetAssetVariables(this, Exports);
-	}
 
 #if WITH_EDITOR
-	// Display programmatic graphs
-	if(CVarDumpProgrammaticGraphs.GetValueOnGameThread())
-	{
-		FUtils::OpenProgrammaticGraphs(this, ProgrammaticGraphs);
-	}
-	else
+		// Display programmatic graphs
+		if(CVarDumpProgrammaticGraphs.GetValueOnGameThread())
+		{
+			FUtils::OpenProgrammaticGraphs(this, ProgrammaticGraphs);
+		}
+		else
 #endif
-	{
-		RemoveProgrammaticGraphs(ProgrammaticGraphs);
-	}
+		{
+			RemoveProgrammaticGraphs(ProgrammaticGraphs);
+		}
 
 #if WITH_EDITOR
-//	RefreshBreakpoints(EditorData);
+		//	RefreshBreakpoints(EditorData);
 #endif
 
-	// Refresh CachedExports
-	if (IAssetRegistry* AssetRegistry = IAssetRegistry::Get())
-	{
-		AssetRegistry->AssetUpdateTags(Asset, EAssetRegistryTagsCaller::Fast);
+		// Refresh CachedExports
+		if (IAssetRegistry* AssetRegistry = IAssetRegistry::Get())
+		{
+			AssetRegistry->AssetUpdateTags(Asset, EAssetRegistryTagsCaller::Fast);
+		}
 	}
 }
 
@@ -948,6 +948,31 @@ void UAnimNextRigVMAssetEditorData::HandleReplaceNotify(UObject* InAsset, const
 	Controller->CloseUndoBracket();
 }
 
+bool UAnimNextRigVMAssetEditorData::IsDirtyForRecompilation() const
+{
+	if(bVMRecompilationRequired)
+	{
+		return true;
+	}
+
+	bool bDependencyDirty = false;
+	ForEachEntryOfType<UAnimNextDataInterfaceEntry>([&bDependencyDirty](UAnimNextDataInterfaceEntry* InEntry)
+	{
+		if(InEntry->DataInterface)
+		{
+			UAnimNextRigVMAssetEditorData* EditorData = UE::AnimNext::UncookedOnly::FUtils::GetEditorData<UAnimNextRigVMAssetEditorData>(InEntry->DataInterface.Get());
+			if(EditorData->IsDirtyForRecompilation())
+			{
+				bDependencyDirty = true;
+				return false;
+			}
+		}
+		return true;
+	});
+
+	return bDependencyDirty;
+}
+
 void UAnimNextRigVMAssetEditorData::RecompileVMIfRequired()
 {
 	if (bVMRecompilationRequired)
@@ -1911,6 +1936,118 @@ void UAnimNextRigVMAssetEditorData::GetAllVariables(TArray<UAnimNextVariableEntr
 	}
 }
 
+void UAnimNextRigVMAssetEditorData::HandleReportFromCompiler(EMessageSeverity::Type InSeverity, UObject* InSubject, const FString& InMessage)
+{
+	FCompilerResultsLog& Log = UE::AnimNext::UncookedOnly::FScopedCompilerResults::GetLog();
+
+	UObject* SubjectForMessage = InSubject;
+	if(URigVMNode* ModelNode = Cast<URigVMNode>(SubjectForMessage))
+	{
+		if(IRigVMClientHost* RigVMClientHost = ModelNode->GetImplementingOuter<IRigVMClientHost>())
+		{
+			if(URigVMNode* OriginalModelNode = Cast<URigVMNode>(Log.FindSourceObject(ModelNode)))
+			{
+				if(URigVMEdGraph* EdGraph = Cast<URigVMEdGraph>(RigVMClientHost->GetEditorObjectForRigVMGraph(OriginalModelNode->GetGraph())))
+				{
+					if(UEdGraphNode* EdNode = EdGraph->FindNodeForModelNodeName(OriginalModelNode->GetFName()))
+					{
+						SubjectForMessage = EdNode;
+					}
+				}
+			}
+		}
+	}
+
+	TSharedPtr<FTokenizedMessage> Message;
+	if (InSeverity == EMessageSeverity::Error)
+	{
+		// see UnitTest "ControlRig.Basics.OrphanedPins" to learn why errors are suppressed this way
+		if (VMCompileSettings.SurpressErrors)
+		{
+			Log.bSilentMode = true;
+		}
+
+		if(InMessage.Contains(TEXT("@@")))
+		{
+			Message = Log.Error(*InMessage, SubjectForMessage);
+		}
+		else
+		{
+			Message = Log.Error(*InMessage);
+		}
+
+		// see UnitTest "ControlRig.Basics.OrphanedPins" to learn why errors are suppressed this way
+		if (!VMCompileSettings.SurpressErrors)
+		{ 
+			FScriptExceptionHandler::Get().HandleException(ELogVerbosity::Error, *InMessage, *FString());
+		}
+
+		bErrorsDuringCompilation = true;
+	}
+	else if (InSeverity == EMessageSeverity::Warning)
+	{
+		if(InMessage.Contains(TEXT("@@")))
+		{
+			Message = Log.Warning(*InMessage, SubjectForMessage);
+		}
+		else
+		{
+			Message = Log.Warning(*InMessage);
+		}
+
+		FScriptExceptionHandler::Get().HandleException(ELogVerbosity::Warning, *InMessage, *FString());
+
+		bWarningsDuringCompilation = true;
+	}
+	else
+	{
+		if(InMessage.Contains(TEXT("@@")))
+		{
+			Message = Log.Note(*InMessage, SubjectForMessage);
+		}
+		else
+		{
+			Message = Log.Note(*InMessage);
+		}
+
+		UE_LOG(LogAnimation, Display, TEXT("%s"), *InMessage);
+	}
+
+	if (URigVMEdGraphNode* EdGraphNode = Cast<URigVMEdGraphNode>(SubjectForMessage))
+	{
+		if(Message.IsValid())
+		{
+			EdGraphNode->SetErrorInfo(InSeverity, Message->ToText().ToString());
+		}
+		else
+		{
+			EdGraphNode->SetErrorInfo(InSeverity, InMessage);
+		}
+
+		EdGraphNode->bHasCompilerMessage = EdGraphNode->ErrorType <= int32(EMessageSeverity::Info);
+	}
+}
+
+void UAnimNextRigVMAssetEditorData::ClearErrorInfoForAllEdGraphs()
+{
+	for (UEdGraph* Graph : GetAllEdGraphs())
+	{
+		URigVMEdGraph* RigGraph = Cast<URigVMEdGraph>(Graph);
+		if (RigGraph == nullptr)
+		{
+			continue;
+		}
+
+		for (UEdGraphNode* GraphNode : Graph->Nodes)
+		{
+			if (URigVMEdGraphNode* RigVMEdGraphNode = Cast<URigVMEdGraphNode>(GraphNode))
+			{
+				RigVMEdGraphNode->ClearErrorInfo();
+			}
+		}
+	}
+}
+
 void UAnimNextRigVMAssetEditorData::RefreshExternalModels()
 {
 	GraphModels.Reset();
diff --git a/Engine/Plugins/Experimental/AnimNext/Source/AnimNextUncookedOnly/Private/AnimNextScopedCompilerResults.cpp b/Engine/Plugins/Experimental/AnimNext/Source/AnimNextUncookedOnly/Private/AnimNextScopedCompilerResults.cpp
new file mode 100644
--- /dev/null
+++ b/Engine/Plugins/Experimental/AnimNext/Source/AnimNextUncookedOnly/Private/AnimNextScopedCompilerResults.cpp
@@ -0,0 +1,134 @@
+﻿// Copyright Epic Games, Inc. All Rights Reserved.
+
+#include "AnimNextScopedCompilerResults.h"
+
+#include "AnimNextRigVMAsset.h"
+#include "AnimNextRigVMAssetEditorData.h"
+#include "UncookedOnlyUtils.h"
+#include "Kismet2/CompilerResultsLog.h"
+#include "Logging/MessageLog.h"
+
+#define LOCTEXT_NAMESPACE "AnimNextScopedCompilerResults"
+
+namespace UE::AnimNext::UncookedOnly
+{
+
+struct FCompilerResultsThreadData
+{
+	TArray<TWeakPtr<FCompilerResultsLog>> CurrentLogStack;
+	TArray<TSharedRef<FTokenizedMessage>> Messages;
+	int32 NumErrors = 0;
+	int32 NumWarnings = 0;
+	bool bDidCompile = false;
+};
+
+static thread_local FCompilerResultsThreadData GCompilerResultsData;
+
+FScopedCompilerResults::FScopedCompilerResults(const FText& InJobName)
+	: FScopedCompilerResults(InJobName, nullptr, TArrayView<UObject*>())
+{
+}
+
+FScopedCompilerResults::FScopedCompilerResults(UObject* InObject)
+	: FScopedCompilerResults(FText::FromName(InObject->GetFName()), InObject, TArrayView<UObject*>(&InObject, 1))
+{
+}
+
+FScopedCompilerResults::FScopedCompilerResults(const FText& InJobName, UObject* InObject, TArrayView<UObject*> InAssets)
+	: JobName(InJobName)
+	, Object(InObject)
+{
+	FCompilerResultsThreadData* ThreadData = &GCompilerResultsData;
+	StartTime = FPlatformTime::Seconds();
+	ThreadData->bDidCompile = Object != nullptr;
+	Log = MakeShared<FCompilerResultsLog>(); 
+	ThreadData->CurrentLogStack.Push(Log);
+
+	for(UObject* Asset : InAssets)
+	{
+		UAnimNextRigVMAsset* AnimNextRigVMAsset = Cast<UAnimNextRigVMAsset>(Asset);
+		if(AnimNextRigVMAsset == nullptr)
+		{
+			return;
+		}
+
+		UAnimNextRigVMAssetEditorData* EditorData = FUtils::GetEditorData<UAnimNextRigVMAssetEditorData>(AnimNextRigVMAsset);
+		if(EditorData == nullptr)
+		{
+			return;
+		}
+
+		EditorData->ClearErrorInfoForAllEdGraphs();
+	}
+}
+
+FScopedCompilerResults::~FScopedCompilerResults()
+{
+	FinishTime = FPlatformTime::Seconds();
+
+	FCompilerResultsThreadData* ThreadData = &GCompilerResultsData;
+	
+	// Accumulate messages
+	ThreadData->Messages.Append(Log->Messages);
+	ThreadData->NumErrors += Log->NumErrors;
+	ThreadData->NumWarnings += Log->NumWarnings;
+
+	ThreadData->CurrentLogStack.Pop();
+
+	if(ThreadData->CurrentLogStack.IsEmpty())
+	{
+		if(ThreadData->bDidCompile)
+		{
+			FMessageLog MessageLog("AnimNextCompilerResults");
+
+			MessageLog.NewPage(FText::Format(LOCTEXT("CompileFormat", "Compile {0}: {1}"), JobName, FText::AsDateTime(FDateTime::Now())));
+			MessageLog.AddMessages(ThreadData->Messages);
+
+			// Print summary
+			FNumberFormattingOptions TimeFormat;
+			TimeFormat.MaximumFractionalDigits = 2;
+			TimeFormat.MinimumFractionalDigits = 2;
+			TimeFormat.MaximumIntegralDigits = 4;
+			TimeFormat.MinimumIntegralDigits = 4;
+			TimeFormat.UseGrouping = false;
+
+			FFormatNamedArguments Args;
+			Args.Add(TEXT("CurrentTime"), FText::AsNumber(FinishTime - GStartTime, &TimeFormat));
+			Args.Add(TEXT("JobName"), JobName);
+			Args.Add(TEXT("CompileTime"), (int32)((FinishTime - StartTime) * 1000));
+			Args.Add(TEXT("ObjectPath"), Object != nullptr ? FText::Format(LOCTEXT("ObjectPathFormat", "({0})"), FText::FromString(Object->GetPathName())) : FText::GetEmpty());
+
+			if (ThreadData->NumErrors > 0)
+			{
+				Args.Add(TEXT("NumErrors"), ThreadData->NumErrors);
+				Args.Add(TEXT("NumWarnings"), ThreadData->NumWarnings);
+				MessageLog.Info(FText::Format(LOCTEXT("CompileFailed", "[{CurrentTime}] Compile of {JobName} failed. {NumErrors} Error(s), {NumWarnings} Warning(s) [in {CompileTime} ms] {ObjectPath}"), MoveTemp(Args)));
+			}
+			else if(ThreadData->NumWarnings > 0)
+			{
+				Args.Add(TEXT("NumWarnings"), ThreadData->NumWarnings);
+				MessageLog.Info(FText::Format(LOCTEXT("CompileWarning", "[{CurrentTime}] Compile of {JobName} successful. {NumWarnings} Warning(s) [in {CompileTime} ms] {ObjectPath}"), MoveTemp(Args)));
+			}
+			else
+			{
+				MessageLog.Info(FText::Format(LOCTEXT("CompileSuccess", "[{CurrentTime}] Compile of {JobName} successful! [in {CompileTime} ms] {ObjectPath}"), MoveTemp(Args)));
+			}
+		}
+
+		ThreadData->Messages.Empty();
+		ThreadData->NumErrors = 0;
+		ThreadData->NumWarnings = 0;
+		ThreadData->bDidCompile = false;
+	}
+}
+
+FCompilerResultsLog& FScopedCompilerResults::GetLog()
+{
+	FCompilerResultsThreadData* ThreadData = &GCompilerResultsData;
+	check(ThreadData->CurrentLogStack.Num() && ThreadData->CurrentLogStack.Top().IsValid());
+	return *ThreadData->CurrentLogStack.Top().Pin().Get();
+}
+
+}
+
+#undef LOCTEXT_NAMESPACE
\ No newline at end of file
diff --git a/Engine/Plugins/Experimental/AnimNext/Source/AnimNextUncookedOnly/Private/Entries/AnimNextDataInterfaceEntry.cpp b/Engine/Plugins/Experimental/AnimNext/Source/AnimNextUncookedOnly/Private/Entries/AnimNextDataInterfaceEntry.cpp
--- a/Engine/Plugins/Experimental/AnimNext/Source/AnimNextUncookedOnly/Private/Entries/AnimNextDataInterfaceEntry.cpp
+++ b/Engine/Plugins/Experimental/AnimNext/Source/AnimNextUncookedOnly/Private/Entries/AnimNextDataInterfaceEntry.cpp
@@ -352,10 +352,11 @@ void UAnimNextDataInterfaceEntry::HandleDataInterfaceModified(UAnimNextRigVMAsse
 	case EAnimNextEditorDataNotifType::EntryRemoved:
 	case EAnimNextEditorDataNotifType::EntryRenamed:
 	case EAnimNextEditorDataNotifType::EntryAccessSpecifierChanged:
+	case EAnimNextEditorDataNotifType::VariableTypeChanged:
 	case EAnimNextEditorDataNotifType::VariableDefaultValueChanged:
 		if(UAnimNextRigVMAssetEditorData* EditorData = GetTypedOuter<UAnimNextRigVMAssetEditorData>())
 		{
-			EditorData->RecompileVM();
+			EditorData->RequestAutoVMRecompilation();
 		}
 		break;
 	default:
diff --git a/Engine/Plugins/Experimental/AnimNext/Source/AnimNextUncookedOnly/Private/Entries/AnimNextRigVMAssetEntry.cpp b/Engine/Plugins/Experimental/AnimNext/Source/AnimNextUncookedOnly/Private/Entries/AnimNextRigVMAssetEntry.cpp
--- a/Engine/Plugins/Experimental/AnimNext/Source/AnimNextUncookedOnly/Private/Entries/AnimNextRigVMAssetEntry.cpp
+++ b/Engine/Plugins/Experimental/AnimNext/Source/AnimNextUncookedOnly/Private/Entries/AnimNextRigVMAssetEntry.cpp
@@ -2,8 +2,11 @@
 
 #include "Entries/AnimNextRigVMAssetEntry.h"
 #include "AnimNextRigVMAssetEditorData.h"
+#include "AnimNextScopedCompilerResults.h"
 #include "Misc/TransactionObjectEvent.h"
 
+#define LOCTEXT_NAMESPACE "AnimNextRigVMAssetEntry"
+
 void UAnimNextRigVMAssetEntry::Initialize(UAnimNextRigVMAssetEditorData* InEditorData)
 {
 	InEditorData->RigVMGraphModifiedEvent.RemoveAll(this);
@@ -42,6 +45,11 @@ void UAnimNextRigVMAssetEntry::BroadcastModified(EAnimNextEditorDataNotifType In
 {
 	if(UAnimNextRigVMAssetEditorData* EditorData = Cast<UAnimNextRigVMAssetEditorData>(GetOuter()))
 	{
+		// Modifications here can trigger compilation, so add a scope to catch compile batches
+		UE::AnimNext::UncookedOnly::FScopedCompilerResults CompilerResults(LOCTEXT("ModifiedAssetsJobName", "Modified Assets"));
+
 		EditorData->BroadcastModified(InType, this);
 	}
-}
\ No newline at end of file
+}
+
+#undef LOCTEXT_NAMESPACE
\ No newline at end of file
diff --git a/Engine/Plugins/Experimental/AnimNext/Source/AnimNextUncookedOnly/Public/AnimNextScopedCompilerResults.h b/Engine/Plugins/Experimental/AnimNext/Source/AnimNextUncookedOnly/Public/AnimNextScopedCompilerResults.h
new file mode 100644
--- /dev/null
+++ b/Engine/Plugins/Experimental/AnimNext/Source/AnimNextUncookedOnly/Public/AnimNextScopedCompilerResults.h
@@ -0,0 +1,34 @@
+﻿// Copyright Epic Games, Inc. All Rights Reserved.
+
+#pragma once
+
+#include "CoreMinimal.h"
+#include "Kismet2/CompilerResultsLog.h"
+
+class UAnimNextRigVMAssetEditorData;
+class FCompilerResultsLog;
+
+namespace UE::AnimNext::UncookedOnly
+{
+
+// RAII helper for scoping output of compiler results
+class FScopedCompilerResults
+{
+public:
+	ANIMNEXTUNCOOKEDONLY_API explicit FScopedCompilerResults(const FText& InJobName);
+	ANIMNEXTUNCOOKEDONLY_API explicit FScopedCompilerResults(UObject* InObject);
+	ANIMNEXTUNCOOKEDONLY_API FScopedCompilerResults(const FText& InJobName, UObject* InObject, TArrayView<UObject*> InAssets);
+	ANIMNEXTUNCOOKEDONLY_API ~FScopedCompilerResults();
+
+	// Get the log that is currently in scope, asserts if no scope is active
+	ANIMNEXTUNCOOKEDONLY_API static FCompilerResultsLog& GetLog();
+
+private:
+	TSharedPtr<FCompilerResultsLog> Log;
+	FText JobName;
+	UObject* Object = nullptr;
+	double StartTime = 0.0;
+	double FinishTime = 0.0;
+};
+
+}
diff --git a/Engine/Plugins/Experimental/AnimNextAnimGraph/Source/AnimNextAnimGraphEditor/Private/Graph/STraitEditorView.cpp b/Engine/Plugins/Experimental/AnimNextAnimGraph/Source/AnimNextAnimGraphEditor/Private/Graph/STraitEditorView.cpp
--- a/Engine/Plugins/Experimental/AnimNextAnimGraph/Source/AnimNextAnimGraphEditor/Private/Graph/STraitEditorView.cpp
+++ b/Engine/Plugins/Experimental/AnimNextAnimGraph/Source/AnimNextAnimGraphEditor/Private/Graph/STraitEditorView.cpp
@@ -279,8 +279,7 @@ void STraitEditorView::RefreshTraitStackTraitsStatus()
 
 			TArray<TSharedPtr<FTraitDataEditorDef>>& CurrentTraitsData = *(TraitEditorSharedData->CurrentTraitsDataShared.Get());
 
-			FMessageLog Log("AnimNextCompilerResults");
-			Log.NewPage(LOCTEXT("TraitStackCompileResults", "Trait Stack Compilation"));
+			TArray<TSharedRef<FTokenizedMessage>> Messages;
 
 			const int32 NumTraits = CurrentTraitsData.Num();
 			for (int32 i = NumTraits - 1; i >= 0; i--)
@@ -305,7 +304,7 @@ void STraitEditorView::RefreshTraitStackTraitsStatus()
 					{
 						const EMessageSeverity::Type Severity = (StatusMessage.Status == FTraitStackTraitStatus::EStackStatus::Warning) ? EMessageSeverity::Warning : EMessageSeverity::Error;
 						const FText& MessageText = StatusMessage.MessageText;
-						Log.AddMessage(FTokenizedMessage::Create(Severity, FText::Format(LOCTEXT("TraitEditorLogTraitNameErrorFormat", "{0}: {1}"), TraitData->TraitDisplayName, MessageText)));
+						Messages.Add(FTokenizedMessage::Create(Severity, FText::Format(LOCTEXT("TraitEditorLogTraitNameErrorFormat", "{0}: {1}"), TraitData->TraitDisplayName, MessageText)));
 					}
 
 					// Open tab to display errors
@@ -315,6 +314,13 @@ void STraitEditorView::RefreshTraitStackTraitsStatus()
 					}
 				}
 			}
+
+			if(Messages.Num())
+			{
+				FMessageLog Log("AnimNextCompilerResults");
+				Log.NewPage(LOCTEXT("TraitStackCompileResults", "Trait Stack Compilation"));
+				Log.AddMessages(Messages);
+			}
 		}
 		else
 		{
diff --git a/Engine/Plugins/Experimental/AnimNextAnimGraph/Source/AnimNextAnimGraphUncookedOnly/Private/Graph/AnimNextAnimationGraph_EditorData.cpp b/Engine/Plugins/Experimental/AnimNextAnimGraph/Source/AnimNextAnimGraphUncookedOnly/Private/Graph/AnimNextAnimationGraph_EditorData.cpp
--- a/Engine/Plugins/Experimental/AnimNextAnimGraph/Source/AnimNextAnimGraphUncookedOnly/Private/Graph/AnimNextAnimationGraph_EditorData.cpp
+++ b/Engine/Plugins/Experimental/AnimNextAnimGraph/Source/AnimNextAnimGraphUncookedOnly/Private/Graph/AnimNextAnimationGraph_EditorData.cpp
@@ -3,6 +3,7 @@
 #include "Graph/AnimNextAnimationGraph_EditorData.h"
 #include "AnimGraphUncookedOnlyUtils.h"
 #include "AnimNextAnimGraphWorkspaceAssetUserData.h"
+#include "AnimNextScopedCompilerResults.h"
 #include "UncookedOnlyUtils.h"
 #include "AssetRegistry/IAssetRegistry.h"
 #include "Curves/CurveFloat.h"
@@ -30,6 +31,8 @@
 #include "Graph/RigUnit_AnimNextGraphEvaluator.h"
 #include "Graph/RigUnit_AnimNextGraphRoot.h"
 
+#define LOCTEXT_NAMESPACE "AnimNextAnimationGraph_EditorData"
+
 namespace UE::AnimNext::UncookedOnly::Private
 {
 	// Represents a trait entry on a node
@@ -468,6 +471,7 @@ void UAnimNextAnimationGraph_EditorData::RecompileVM()
 	TGuardValue<bool> CompilingGuard(bIsCompiling, true);
 
 	UAnimNextAnimationGraph* AnimationGraph = FUtils::GetAsset<UAnimNextAnimationGraph>(this);
+	FScopedCompilerResults CompilerResults(AnimationGraph);
 
 	// Before we re-compile a graph, we need to release and live instances since we need the metadata we are about to replace
 	// to call trait destructors etc
@@ -475,6 +479,7 @@ void UAnimNextAnimationGraph_EditorData::RecompileVM()
 
 	CachedExports.Reset(); // asset variables and other tags will be updated at the end by AssetRegistry->AssetUpdateTags
 
+	bWarningsDuringCompilation = false;
 	bErrorsDuringCompilation = false;
 
 	RigGraphDisplaySettings.MinMicroSeconds = RigGraphDisplaySettings.LastMinMicroSeconds = DBL_MAX;
@@ -487,13 +492,9 @@ void UAnimNextAnimationGraph_EditorData::RecompileVM()
 
 		VMCompileSettings.SetExecuteContextStruct(FAnimNextExecuteContext::StaticStruct());
 		FRigVMCompileSettings Settings = (bCompileInDebugMode) ? FRigVMCompileSettings::Fast(VMCompileSettings.GetExecuteContextStruct()) : VMCompileSettings;
+		Settings.SurpressInfoMessages = false;
 		Settings.ASTSettings.bSetupTraits = false; // disable the default implementation of decorators for now
-
-		FMessageLog("AnimNextCompilerResults").NewPage(FText::FromName(AnimationGraph->GetFName()));
-		Settings.ASTSettings.ReportDelegate.BindLambda([](EMessageSeverity::Type InType, UObject* InObject, const FString& InString)
-		{
-			FMessageLog("AnimNextCompilerResults").Message(InType, FText::FromString(InString));
-		});
+		Settings.ASTSettings.ReportDelegate.BindUObject(this, &UAnimNextAnimationGraph_EditorData::HandleReportFromCompiler);
 
 		FUtils::RecreateVM(AnimationGraph);
 
@@ -532,7 +533,16 @@ void UAnimNextAnimationGraph_EditorData::RecompileVM()
 			// We use a temporary graph models to build our final graphs that we'll compile
 			if(SourceGraph->GetSchemaClass() == UAnimNextAnimationGraphSchema::StaticClass())
 			{
-				URigVMGraph* TempGraph = CastChecked<URigVMGraph>(StaticDuplicateObject(SourceGraph, this, NAME_None, RF_Transient));
+				TMap<UObject*, UObject*> CreatedObjects;
+				FObjectDuplicationParameters Parameters = InitStaticDuplicateObjectParams(SourceGraph, this, NAME_None, RF_Transient);
+				Parameters.CreatedObjects = &CreatedObjects;
+				URigVMGraph* TempGraph = CastChecked<URigVMGraph>(StaticDuplicateObjectEx(Parameters));
+				TempGraph->SetExternalPackage(nullptr);
+				for(URigVMNode* SourceNode : SourceGraph->GetNodes())
+				{
+					FScopedCompilerResults::GetLog().NotifyIntermediateObjectCreation(CreatedObjects.FindChecked(SourceNode), SourceNode);
+				}
+
 				UAnimNextController* TempController = CastChecked<UAnimNextController>(VMClient->GetOrCreateController(TempGraph));
 				TempGraph->SetFlags(RF_Transient);
 				AnimGraphs.Add(TempGraph);
@@ -672,28 +682,28 @@ void UAnimNextAnimationGraph_EditorData::RecompileVM()
 
 		// Now that the graph has been re-compiled, re-allocate the previous live instances
 		AnimationGraph->ThawGraphInstances();
-	}
 
 #if WITH_EDITOR
-	// Display programmatic graphs
-	if(CVarDumpProgrammaticGraphs.GetValueOnGameThread())
-	{
-		FUtils::OpenProgrammaticGraphs(this, ProgrammaticGraphs);
-	}
-	else
+		// Display programmatic graphs
+		if(CVarDumpProgrammaticGraphs.GetValueOnGameThread())
+		{
+			FUtils::OpenProgrammaticGraphs(this, ProgrammaticGraphs);
+		}
+		else
 #endif
-	{
-		RemoveProgrammaticGraphs(ProgrammaticGraphs);
-	}
+		{
+			RemoveProgrammaticGraphs(ProgrammaticGraphs);
+		}
 
 #if WITH_EDITOR
-//	RefreshBreakpoints(EditorData);
+		//	RefreshBreakpoints(EditorData);
 #endif
 
-	// Refresh CachedExports, also updates variables at UAnimNextRigVMAssetEditorData::GetAssetRegistryTags
-	if (IAssetRegistry* AssetRegistry = IAssetRegistry::Get())
-	{
-		AssetRegistry->AssetUpdateTags(GetTypedOuter<UAnimNextAnimationGraph>(), EAssetRegistryTagsCaller::Fast);
+		// Refresh CachedExports, also updates variables at UAnimNextRigVMAssetEditorData::GetAssetRegistryTags
+		if (IAssetRegistry* AssetRegistry = IAssetRegistry::Get())
+		{
+			AssetRegistry->AssetUpdateTags(GetTypedOuter<UAnimNextAnimationGraph>(), EAssetRegistryTagsCaller::Fast);
+		}
 	}
 }
 
@@ -953,3 +963,5 @@ void UAnimNextAnimationGraph_EditorData::GetAnimNextAssetRegistryTags(FAssetRegi
 {
 	UE::AnimNext::UncookedOnly::FAnimGraphUtils::GetAssetManifestNodesRegistryExports(this, OutExports);
 }
+
+#undef LOCTEXT_NAMESPACE
\ No newline at end of file
diff --git a/Engine/Plugins/Experimental/AnimNextStateTree/Source/AnimNextStateTreeEditor/Private/AnimNextStateTreeEditorHost.cpp b/Engine/Plugins/Experimental/AnimNextStateTree/Source/AnimNextStateTreeEditor/Private/AnimNextStateTreeEditorHost.cpp
--- a/Engine/Plugins/Experimental/AnimNextStateTree/Source/AnimNextStateTreeEditor/Private/AnimNextStateTreeEditorHost.cpp
+++ b/Engine/Plugins/Experimental/AnimNextStateTree/Source/AnimNextStateTreeEditor/Private/AnimNextStateTreeEditorHost.cpp
@@ -14,14 +14,14 @@ void FAnimNextStateTreeEditorHost::Init(const TWeakPtr<UE::Workspace::IWorkspace
 	
 	TSharedPtr<UE::Workspace::IWorkspaceEditor> SharedEditor = InWeakWorkspaceEditor.Pin();
 	check(SharedEditor.IsValid());
-	SharedEditor->OnFocussedAssetChanged().AddSP(this, &FAnimNextStateTreeEditorHost::OnWorkspaceFocussedAssetChanged);
+	SharedEditor->OnFocussedDocumentChanged().AddSP(this, &FAnimNextStateTreeEditorHost::OnWorkspaceFocussedDocumentChanged);
 }
 
 UStateTree* FAnimNextStateTreeEditorHost::GetStateTree() const
 {
 	if (TSharedPtr<UE::Workspace::IWorkspaceEditor> SharedWorkspaceEditor = WeakWorkspaceEditor.Pin())
 	{
-		const TObjectPtr<UAnimNextStateTree> AnimNextStateTreePtr = SharedWorkspaceEditor->GetFocussedAsset<UAnimNextStateTree>();
+		const TObjectPtr<UAnimNextStateTree> AnimNextStateTreePtr = SharedWorkspaceEditor->GetFocussedDocument<UAnimNextStateTree>();
 		return AnimNextStateTreePtr ? AnimNextStateTreePtr->StateTree : nullptr;
 	}	
 
@@ -38,7 +38,12 @@ FName FAnimNextStateTreeEditorHost::GetCompilerTabName() const
 	return UE::AnimNext::Editor::CompilerResultsTabName;
 }
 
-void FAnimNextStateTreeEditorHost::OnWorkspaceFocussedAssetChanged(TObjectPtr<UObject> InObject) const
+bool FAnimNextStateTreeEditorHost::ShouldShowCompileButton() const
+{
+	return false;
+}
+
+void FAnimNextStateTreeEditorHost::OnWorkspaceFocussedDocumentChanged(TObjectPtr<UObject> InObject) const
 {
 	OnStateTreeChangedDelegate.Broadcast();
 }
diff --git a/Engine/Plugins/Experimental/AnimNextStateTree/Source/AnimNextStateTreeEditor/Private/AnimNextStateTreeEditorHost.h b/Engine/Plugins/Experimental/AnimNextStateTree/Source/AnimNextStateTreeEditor/Private/AnimNextStateTreeEditorHost.h
--- a/Engine/Plugins/Experimental/AnimNextStateTree/Source/AnimNextStateTreeEditor/Private/AnimNextStateTreeEditorHost.h
+++ b/Engine/Plugins/Experimental/AnimNextStateTree/Source/AnimNextStateTreeEditor/Private/AnimNextStateTreeEditorHost.h
@@ -18,13 +18,13 @@ public:
 	virtual UStateTree* GetStateTree() const override;
 	virtual FName GetCompilerLogName() const override;
 	virtual FName GetCompilerTabName() const override;
-	
+	virtual bool ShouldShowCompileButton() const override;
 	virtual FSimpleMulticastDelegate& OnStateTreeChanged() override;
 	virtual TSharedPtr<IDetailsView> GetAssetDetailsView() override;
 	virtual TSharedPtr<IDetailsView> GetDetailsView() override;
 
 protected:
-	void OnWorkspaceFocussedAssetChanged(TObjectPtr<UObject> InObject) const;
+	void OnWorkspaceFocussedDocumentChanged(TObjectPtr<UObject> InObject) const;
 
 	TWeakPtr<UE::Workspace::IWorkspaceEditor> WeakWorkspaceEditor;	
 	FSimpleMulticastDelegate OnStateTreeChangedDelegate;
diff --git a/Engine/Plugins/Experimental/AnimNextStateTree/Source/AnimNextStateTreeEditor/Private/AnimNextStateTreeEditorModule.cpp b/Engine/Plugins/Experimental/AnimNextStateTree/Source/AnimNextStateTreeEditor/Private/AnimNextStateTreeEditorModule.cpp
--- a/Engine/Plugins/Experimental/AnimNextStateTree/Source/AnimNextStateTreeEditor/Private/AnimNextStateTreeEditorModule.cpp
+++ b/Engine/Plugins/Experimental/AnimNextStateTree/Source/AnimNextStateTreeEditor/Private/AnimNextStateTreeEditorModule.cpp
@@ -23,6 +23,7 @@
 #include "ToolMenus.h"
 #include "Framework/MultiBox/ToolMenuBase.h"
 #include "AnimNextRigVMAssetEditorData.h"
+#include "StateTreeAssetCompilationHandler.h"
 
 #define LOCTEXT_NAMESPACE "AnimNextStateTreeEditorModule"
 
@@ -33,6 +34,12 @@ void FAnimNextStateTreeEditorModule::StartupModule()
 	// Register StateTree as supported asset in AnimNext workspaces
 	Editor::IAnimNextEditorModule& AnimNextEditorModule = FModuleManager::Get().LoadModuleChecked<Editor::IAnimNextEditorModule>("AnimNextEditor");	
 	AnimNextEditorModule.AddWorkspaceSupportedAssetClass(UAnimNextStateTree::StaticClass()->GetClassPathName());
+	AnimNextEditorModule.RegisterAssetCompilationHandler(UAnimNextStateTree::StaticClass()->GetClassPathName(), Editor::FAssetCompilationHandlerFactoryDelegate::CreateLambda([](UObject* InAsset)
+	{
+		TSharedRef<FStateTreeAssetCompilationHandler> CompilationHandler = MakeShared<FStateTreeAssetCompilationHandler>(InAsset);
+		CompilationHandler->Initialize();
+		return CompilationHandler;
+	}));
 
 	// Extend Workspace Editor layout to deal with StateTreeEditorMode tabs
 	Workspace::IWorkspaceEditorModule& WorkspaceEditorModule = FModuleManager::Get().LoadModuleChecked<Workspace::IWorkspaceEditorModule>("WorkspaceEditor");
@@ -321,6 +328,7 @@ void FAnimNextStateTreeEditorModule::ShutdownModule()
 		if( Editor::IAnimNextEditorModule* AnimNextEditorModule = FModuleManager::Get().GetModulePtr<Editor::IAnimNextEditorModule>("AnimNextEditor"))
 		{
 			AnimNextEditorModule->RemoveWorkspaceSupportedAssetClass(UStateTree::StaticClass()->GetClassPathName());
+			AnimNextEditorModule->UnregisterAssetCompilationHandler(UAnimNextStateTree::StaticClass()->GetClassPathName());
 		}
 	}
 }
diff --git a/Engine/Plugins/Experimental/AnimNextStateTree/Source/AnimNextStateTreeEditor/Private/StateTreeAssetCompilationHandler.cpp b/Engine/Plugins/Experimental/AnimNextStateTree/Source/AnimNextStateTreeEditor/Private/StateTreeAssetCompilationHandler.cpp
new file mode 100644
--- /dev/null
+++ b/Engine/Plugins/Experimental/AnimNextStateTree/Source/AnimNextStateTreeEditor/Private/StateTreeAssetCompilationHandler.cpp
@@ -0,0 +1,147 @@
+﻿// Copyright Epic Games, Inc. All Rights Reserved.
+
+#include "StateTreeAssetCompilationHandler.h"
+
+#include "AnimNextRigVMAssetEditorData.h"
+#include "AnimNextScopedCompilerResults.h"
+#include "AnimNextStateTree.h"
+#include "AnimNextStateTree_EditorData.h"
+#include "Editor.h"
+#include "StateTree.h"
+#include "StateTreeCompilerLog.h"
+#include "StateTreeEditingSubsystem.h"
+#include "UncookedOnlyUtils.h"
+
+namespace UE::AnimNext::StateTree
+{
+
+FStateTreeAssetCompilationHandler::FStateTreeAssetCompilationHandler(UObject* InAsset)
+	: FAssetCompilationHandler(InAsset)
+{
+	UAnimNextStateTree* AnimNextStateTree = Cast<UAnimNextStateTree>(InAsset);
+	CachedStateTree = AnimNextStateTree->StateTree;
+	UStateTreeEditingSubsystem::ValidateStateTree(AnimNextStateTree->StateTree);
+	EditorDataHash = AnimNextStateTree->StateTree->LastCompiledEditorDataHash;
+}
+
+void FStateTreeAssetCompilationHandler::Initialize()
+{
+	if (UStateTreeEditingSubsystem* StateTreeEditingSubsystem = GEditor->GetEditorSubsystem<UStateTreeEditingSubsystem>())
+	{
+		TSharedRef<FStateTreeViewModel> ViewModel = StateTreeEditingSubsystem->FindOrAddViewModel(CachedStateTree.Get());
+		ViewModel->GetOnAssetChanged().AddSP(this, &FStateTreeAssetCompilationHandler::HandleAssetChanged);
+		ViewModel->GetOnStatesChanged().AddSP(this, &FStateTreeAssetCompilationHandler::HandleStatesChanged);
+		ViewModel->GetOnStateAdded().AddSP(this, &FStateTreeAssetCompilationHandler::HandleStateAdded);
+		ViewModel->GetOnStatesRemoved().AddSP(this, &FStateTreeAssetCompilationHandler::HandleStatesRemoved);
+		ViewModel->GetOnStatesMoved().AddSP(this, &FStateTreeAssetCompilationHandler::HandleStatesMoved);
+	}
+}
+
+void FStateTreeAssetCompilationHandler::Compile(TSharedRef<Workspace::IWorkspaceEditor> InWorkspaceEditor, UObject* InAsset)
+{
+	UAnimNextStateTree* AnimNextStateTree = Cast<UAnimNextStateTree>(InAsset);
+	if (AnimNextStateTree == nullptr)
+	{
+		return;
+	}
+
+	UAnimNextStateTree_EditorData* EditorData = UncookedOnly::FUtils::GetEditorData<UAnimNextStateTree_EditorData>(AnimNextStateTree);
+	if (EditorData == nullptr)
+	{
+		return;
+	}
+
+	UE::AnimNext::UncookedOnly::FScopedCompilerResults CompilerResults(AnimNextStateTree);
+
+	// We can't recompile state tree while in PIE
+	if (!GEditor->IsPlaySessionInProgress())
+	{
+		// First compile the state tree
+		FStateTreeCompilerLog Log;
+		bLastCompileSucceeded = UStateTreeEditingSubsystem::CompileStateTree(AnimNextStateTree->StateTree, Log);
+		EditorDataHash = AnimNextStateTree->StateTree->LastCompiledEditorDataHash;
+
+		// Add to messages
+		for(const TSharedRef<FTokenizedMessage>& Message : Log.ToTokenizedMessages())
+		{
+			CompilerResults.GetLog().AddTokenizedMessage(Message);
+		}
+	}
+	else
+	{
+		CompilerResults.GetLog().Note(*FString::Printf(TEXT("StateTree '%s' compilation was skipped because PIE was running"), *InAsset->GetName()));
+	}
+
+	// Call super to compile the RigVM Asset 
+	UE::AnimNext::Editor::FAssetCompilationHandler::Compile(InWorkspaceEditor, InAsset);
+}
+
+UE::AnimNext::Editor::ECompileStatus FStateTreeAssetCompilationHandler::GetCompileStatus(TSharedRef<Workspace::IWorkspaceEditor> InWorkspaceEditor, const UObject* InAsset) const
+{
+	const UAnimNextStateTree* AnimNextStateTree = Cast<UAnimNextStateTree>(InAsset);
+	if (AnimNextStateTree == nullptr)
+	{
+		return UE::AnimNext::Editor::ECompileStatus::Unknown;
+	}
+
+	UE::AnimNext::Editor::ECompileStatus RigVMAssetStatus = UE::AnimNext::Editor::FAssetCompilationHandler::GetCompileStatus(InWorkspaceEditor, InAsset);
+	UE::AnimNext::Editor::ECompileStatus StateTreeAssetStatus = UE::AnimNext::Editor::ECompileStatus::Unknown;
+	if (AnimNextStateTree->StateTree)
+	{
+		const bool bCompiledDataResetDuringLoad = AnimNextStateTree->StateTree->LastCompiledEditorDataHash == EditorDataHash && !AnimNextStateTree->StateTree->IsReadyToRun();
+		if (!bLastCompileSucceeded || bCompiledDataResetDuringLoad)
+		{
+			StateTreeAssetStatus = UE::AnimNext::Editor::ECompileStatus::Error;
+		}
+		else if (AnimNextStateTree->StateTree->LastCompiledEditorDataHash != EditorDataHash)
+		{
+			StateTreeAssetStatus = UE::AnimNext::Editor::ECompileStatus::Dirty;
+		}
+		else
+		{
+			StateTreeAssetStatus = UE::AnimNext::Editor::ECompileStatus::UpToDate;
+		}
+	}
+
+	return FMath::Max(RigVMAssetStatus, StateTreeAssetStatus);
+}
+
+void FStateTreeAssetCompilationHandler::UpdateCachedInfo()
+{
+	UStateTree* StateTree = CachedStateTree.Get();
+	if(StateTree == nullptr)
+	{
+		return;
+	}
+
+	UStateTreeEditingSubsystem::ValidateStateTree(StateTree);
+	EditorDataHash = UStateTreeEditingSubsystem::CalculateStateTreeHash(StateTree);
+	OnCompileStatusChanged().ExecuteIfBound();
+}
+
+void FStateTreeAssetCompilationHandler::HandleAssetChanged()
+{
+	UpdateCachedInfo();
+}
+
+void FStateTreeAssetCompilationHandler::HandleStatesChanged(const TSet<UStateTreeState*>& /*AffectedStates*/, const FPropertyChangedEvent& /*PropertyChangedEvent*/)
+{
+	UpdateCachedInfo();
+}
+
+void FStateTreeAssetCompilationHandler::HandleStateAdded(UStateTreeState* /*ParentState*/, UStateTreeState* /*NewState*/)
+{
+	UpdateCachedInfo();
+}
+
+void FStateTreeAssetCompilationHandler::HandleStatesRemoved(const TSet<UStateTreeState*>& /*AffectedParents*/)
+{
+	UpdateCachedInfo();
+}
+
+void FStateTreeAssetCompilationHandler::HandleStatesMoved(const TSet<UStateTreeState*>& /*AffectedParents*/, const TSet<UStateTreeState*>& /*MovedStates*/)
+{
+	UpdateCachedInfo();
+}
+
+}
diff --git a/Engine/Plugins/Experimental/AnimNextStateTree/Source/AnimNextStateTreeEditor/Private/StateTreeAssetCompilationHandler.h b/Engine/Plugins/Experimental/AnimNextStateTree/Source/AnimNextStateTreeEditor/Private/StateTreeAssetCompilationHandler.h
new file mode 100644
--- /dev/null
+++ b/Engine/Plugins/Experimental/AnimNextStateTree/Source/AnimNextStateTreeEditor/Private/StateTreeAssetCompilationHandler.h
@@ -0,0 +1,41 @@
+﻿// Copyright Epic Games, Inc. All Rights Reserved.
+
+#pragma once
+
+#include "CoreMinimal.h"
+#include "AssetCompilationHandler.h"
+#include "UObject/WeakObjectPtr.h"
+
+class UStateTreeState;
+class UStateTree;
+
+namespace UE::AnimNext::StateTree
+{
+
+class FStateTreeAssetCompilationHandler : public UE::AnimNext::Editor::FAssetCompilationHandler
+{
+public:
+	explicit FStateTreeAssetCompilationHandler(UObject* InObject);
+
+	// Non-constructor setup (e.g. binding SP delegates)
+	void Initialize();
+
+private:
+	// IAssetCompilationHandler interface
+	virtual void Compile(TSharedRef<Workspace::IWorkspaceEditor> InWorkspaceEditor, UObject* InAsset) override;
+	virtual UE::AnimNext::Editor::ECompileStatus GetCompileStatus(TSharedRef<Workspace::IWorkspaceEditor> InWorkspaceEditor, const UObject* InAsset) const override;
+
+	void UpdateCachedInfo();
+
+	void HandleAssetChanged();
+	void HandleStatesChanged(const TSet<UStateTreeState*>& /*AffectedStates*/, const FPropertyChangedEvent& /*PropertyChangedEvent*/);
+	void HandleStateAdded(UStateTreeState* /*ParentState*/, UStateTreeState* /*NewState*/);
+	void HandleStatesRemoved(const TSet<UStateTreeState*>& /*AffectedParents*/);
+	void HandleStatesMoved(const TSet<UStateTreeState*>& /*AffectedParents*/, const TSet<UStateTreeState*>& /*MovedStates*/);
+
+	TWeakObjectPtr<UStateTree> CachedStateTree;
+	uint32 EditorDataHash = 0;
+	bool bLastCompileSucceeded = true;
+};
+
+}
diff --git a/Engine/Plugins/Experimental/Workspace/Source/WorkspaceEditor/Private/AssetDocumentSummoner.cpp b/Engine/Plugins/Experimental/Workspace/Source/WorkspaceEditor/Private/AssetDocumentSummoner.cpp
--- a/Engine/Plugins/Experimental/Workspace/Source/WorkspaceEditor/Private/AssetDocumentSummoner.cpp
+++ b/Engine/Plugins/Experimental/Workspace/Source/WorkspaceEditor/Private/AssetDocumentSummoner.cpp
@@ -35,7 +35,7 @@ void FAssetDocumentSummoner::OnTabActivated(TSharedPtr<SDockTab> Tab) const
 	TSharedRef<SWorkspaceTabWrapper> TabWrapper = StaticCastSharedRef<SWorkspaceTabWrapper>(Tab->GetContent());
 	if (UObject* DocumentAsset = TabWrapper->GetDocumentObject().Get())
 	{
-		HostingAppPtr.Pin()->SetFocussedAsset(DocumentAsset);
+		HostingAppPtr.Pin()->SetFocussedDocument(DocumentAsset);
 	}
 }
 
diff --git a/Engine/Plugins/Experimental/Workspace/Source/WorkspaceEditor/Private/Workspace.cpp b/Engine/Plugins/Experimental/Workspace/Source/WorkspaceEditor/Private/Workspace.cpp
--- a/Engine/Plugins/Experimental/Workspace/Source/WorkspaceEditor/Private/Workspace.cpp
+++ b/Engine/Plugins/Experimental/Workspace/Source/WorkspaceEditor/Private/Workspace.cpp
@@ -220,12 +220,12 @@ UWorkspaceSchema* UWorkspace::GetSchema() const
 	return SchemaClass->GetDefaultObject<UWorkspaceSchema>();
 }
 
-void UWorkspace::LoadState()
+void UWorkspace::LoadState() const
 {
 	GetState()->LoadFromJson(this);
 }
 
-void UWorkspace::SaveState()
+void UWorkspace::SaveState() const
 {
 	GetState()->SaveToJson(this);
 }
diff --git a/Engine/Plugins/Experimental/Workspace/Source/WorkspaceEditor/Private/Workspace.h b/Engine/Plugins/Experimental/Workspace/Source/WorkspaceEditor/Private/Workspace.h
--- a/Engine/Plugins/Experimental/Workspace/Source/WorkspaceEditor/Private/Workspace.h
+++ b/Engine/Plugins/Experimental/Workspace/Source/WorkspaceEditor/Private/Workspace.h
@@ -14,6 +14,7 @@ class UWorkspaceState;
 class UWorkspaceSchema;
 class UWorkspaceFactory;
 class UAssetDefinition_Workspace;
+struct FInstancedStruct;
 
 namespace UE::Workspace
 {
@@ -104,8 +105,8 @@ class UWorkspace : public UObject
 
 	UWorkspaceSchema* GetSchema() const;
 
-	void LoadState();
-	void SaveState();
+	void LoadState() const;
+	void SaveState() const;
 	UWorkspaceState* GetState() const;
 
 	void ReportError(const TCHAR* InMessage) const;
diff --git a/Engine/Plugins/Experimental/Workspace/Source/WorkspaceEditor/Private/WorkspaceEditor.cpp b/Engine/Plugins/Experimental/Workspace/Source/WorkspaceEditor/Private/WorkspaceEditor.cpp
--- a/Engine/Plugins/Experimental/Workspace/Source/WorkspaceEditor/Private/WorkspaceEditor.cpp
+++ b/Engine/Plugins/Experimental/Workspace/Source/WorkspaceEditor/Private/WorkspaceEditor.cpp
@@ -24,6 +24,8 @@
 #include "EditorModeManager.h"
 #include "EdModeInteractiveToolsContext.h"
 #include "ContextObjectStore.h"
+#include "SGraphDocument.h"
+#include "WorkspaceSchema.h"
 
 #define LOCTEXT_NAMESPACE "WorkspaceEditor"
 
@@ -200,6 +202,7 @@ void FWorkspaceEditor::CreateWidgets()
 void FWorkspaceEditor::PostInitAssetEditor()
 {
 	Workspace->LoadState();
+	GetSchema()->OnLoadWorkspaceState(SharedThis(this), Workspace->GetState()->UserState);
 	RestoreEditedObjectState();
 
 	ExtendMenu();
@@ -272,6 +275,7 @@ void FWorkspaceEditor::SaveEditedObjectState() const
 	DocumentManager->SaveAllState();
 
 	// Persist state
+	Workspace->GetSchema()->OnSaveWorkspaceState(SharedThis(const_cast<FWorkspaceEditor*>(this)), Workspace->GetState()->UserState);
 	Workspace->SaveState();
 }
 
@@ -325,9 +329,48 @@ void FWorkspaceEditor::OpenAssets(TConstArrayView<FAssetData> InAssets)
 
 void FWorkspaceEditor::OpenObjects(TConstArrayView<UObject*> InObjects)
 {
-	for(const UObject* Object : InObjects)
+	for(UObject* Object : InObjects)
 	{
-		OpenDocument(Object, InObjects.Num() == 1 ? FDocumentTracker::EOpenDocumentCause::NavigatingCurrentDocument : FDocumentTracker::EOpenDocumentCause::OpenNewDocument);
+		const FWorkspaceEditorModule& WorkspaceEditorModule = FModuleManager::LoadModuleChecked<FWorkspaceEditorModule>("WorkspaceEditor");
+		const FDocumentSubObjectArgs* DocumentSubObjectArgs = WorkspaceEditorModule.FindDocumentSubObjectType(Object);
+		UObject* OriginalObject = Object;
+		if(DocumentSubObjectArgs && DocumentSubObjectArgs->OnGetDocumentForSubObject.IsBound())
+		{
+			Object = DocumentSubObjectArgs->OnGetDocumentForSubObject.Execute(Object);
+		}
+
+		if(Object)
+		{
+			TSharedPtr<SDockTab> DocumentTab = OpenDocument(Object, InObjects.Num() == 1 ? FDocumentTracker::EOpenDocumentCause::NavigatingCurrentDocument : FDocumentTracker::EOpenDocumentCause::OpenNewDocument);
+			if(DocumentSubObjectArgs && DocumentTab.IsValid())
+			{
+				TSharedRef<SWorkspaceTabWrapper> WorkspaceTabWrapper = StaticCastSharedRef<SWorkspaceTabWrapper>(DocumentTab->GetContent());
+				TSharedRef<SWidget> TabContentWidget = WorkspaceTabWrapper->GetContent();
+
+				// If this is a built-in graph editor widget, supply the inner SGraphEditor
+				if(TabContentWidget->GetType() == SGraphDocument::StaticWidgetClass().GetWidgetType())
+				{
+					TabContentWidget = StaticCastSharedRef<SGraphDocument>(TabContentWidget)->GraphEditor.ToSharedRef();
+					check(TabContentWidget->GetType() == SGraphEditor::StaticWidgetClass().GetWidgetType());
+				}
+
+				DocumentSubObjectArgs->OnPostDocumentOpenedForSubObject.ExecuteIfBound(FWorkspaceEditorContext(SharedThis(this), Object), TabContentWidget, OriginalObject);
+			}
+		}
+	}
+}
+
+void FWorkspaceEditor::GetOpenedAssetsOfClass(TSubclassOf<UObject> InClass, TArray<UObject*>& OutAssets) const
+{
+	FWorkspaceAssetRegistryExports Exports;
+	FWorkspaceEditorModule::GetExportedAssetsForWorkspace(FAssetData(Workspace), Exports);
+	for(const FWorkspaceAssetRegistryExportEntry& Entry : Exports.Assets)
+	{
+		UObject* LoadedObject = Entry.Asset.ResolveObject();
+		if(LoadedObject && LoadedObject->GetClass()->IsChildOf(InClass))
+		{
+			OutAssets.Add(LoadedObject);
+		}
 	}
 }
 
@@ -336,7 +379,7 @@ void FWorkspaceEditor::GetAssets(TArray<FAssetData>& OutAssets) const
 	FWorkspaceAssetRegistryExports Exports;
 	FWorkspaceEditorModule::GetExportedAssetsForWorkspace(FAssetData(Workspace), Exports);
 	IAssetRegistry& AssetRegistry = FModuleManager::LoadModuleChecked<FAssetRegistryModule>("AssetRegistry").Get();
-	for (const FWorkspaceAssetRegistryExportEntry& Entry : Exports.Assets)
+	for(const FWorkspaceAssetRegistryExportEntry& Entry : Exports.Assets)
 	{
 		OutAssets.Add(AssetRegistry.GetAssetByObjectPath(Entry.Asset));
 	}
@@ -407,17 +450,17 @@ void FWorkspaceEditor::SetGlobalSelection(FGlobalSelectionId SelectionId, FOnCle
 	GlobalSelections.Add({ SelectionId, OnClearSelectionDelegate});
 }
 
-void FWorkspaceEditor::SetFocussedAsset(const TObjectPtr<UObject> InAsset)
+void FWorkspaceEditor::SetFocussedDocument(const TObjectPtr<UObject> InDocument)
 {
 	FAssetEditorModeManager* ModeManager = static_cast<FAssetEditorModeManager*>(&GetEditorModeManager());	
 	ModeManager->GetSelectedObjects()->DeselectAll();
-	
-	if (InAsset != nullptr)
+
+	if (InDocument != nullptr)
 	{
-		ModeManager->GetSelectedObjects()->Select(InAsset);
+		ModeManager->GetSelectedObjects()->Select(InDocument);
 
 		FWorkspaceEditorModule& WorkspaceEditorModule = FModuleManager::Get().LoadModuleChecked<FWorkspaceEditorModule>("WorkspaceEditor");
-		if (const FObjectDocumentArgs* Args = WorkspaceEditorModule.FindObjectDocumentType(InAsset))
+		if (const FObjectDocumentArgs* Args = WorkspaceEditorModule.FindObjectDocumentType(InDocument))
 		{
 			const FEditorModeID NewEditorMode = Args->DocumentEditorMode;
 			if (NewEditorMode != NAME_None)
@@ -428,20 +471,20 @@ void FWorkspaceEditor::SetFocussedAsset(const TObjectPtr<UObject> InAsset)
 				}
 			}
 			else
-			{				
+			{
 				ModeManager->DeactivateAllModes();
 			}
-		}		
+		}
 	}
 	else
 	{
 		ModeManager->DeactivateAllModes();
 	}
 
-	OnFocussedAssetChangedDelegate.Broadcast(InAsset);
+	OnFocussedDocumentChangedDelegate.Broadcast(InDocument);
 }
 
-const TObjectPtr<UObject> FWorkspaceEditor::GetFocussedAssetOfClass(const TObjectPtr<UClass> AssetClass) const
+const TObjectPtr<UObject> FWorkspaceEditor::GetFocussedDocumentOfClass(const TObjectPtr<UClass> AssetClass) const
 {
 	FAssetEditorModeManager* ModeManager = static_cast<FAssetEditorModeManager*>(&GetEditorModeManager());	
 	if (USelection* Selection = ModeManager->GetSelectedObjects())
@@ -457,9 +500,9 @@ void FWorkspaceEditor::HandleOutlinerSelectionChanged(TConstArrayView<FWorkspace
 	LastSelectedExports = InExports;
 }
 
-FOnFocussedAssetChanged& FWorkspaceEditor::OnFocussedAssetChanged()
+FOnFocussedDocumentChanged& FWorkspaceEditor::OnFocussedDocumentChanged()
 {
-	return OnFocussedAssetChangedDelegate;
+	return OnFocussedDocumentChangedDelegate;
 }
 
 TSharedPtr<IDetailsView> FWorkspaceEditor::GetDetailsView()
@@ -467,6 +510,11 @@ TSharedPtr<IDetailsView> FWorkspaceEditor::GetDetailsView()
 	return DetailsView;
 }
 
+UObject* FWorkspaceEditor::GetWorkspaceAsset() const
+{
+	return Workspace;
+}
+
 FString FWorkspaceEditor::GetPackageName() const
 {
 	return Workspace->GetPackage()->GetName();
@@ -634,7 +682,7 @@ bool FWorkspaceEditor::CanSaveAsset() const
 	bool bDirtyState = false;
 	for (UObject* EditingObject : GetEditingObjects())
 	{
-		const UPackage* Package = EditingObject->GetOutermost();
+		const UPackage* Package = EditingObject->GetPackage();
 		if (Package->IsDirty() || Package->GetExternalPackages().ContainsByPredicate([](const UPackage* ExternalPackage)  { return ExternalPackage && ExternalPackage->IsDirty(); }))
 		{
 			bDirtyState = true;
@@ -664,9 +712,9 @@ FText FWorkspaceEditor::GetToolkitToolTipText() const
 	check (WorkspaceObject != nullptr);
 	
 	FText FocussedAssetText = FText::FromName(NAME_None);	
-	if (const UObject* FocussedObject = GetFocussedAsset())
+	if (const UObject* FocussedDocument = GetFocussedDocument())
 	{
-		 FocussedAssetText = GetLabelForObject(FocussedObject);
+		 FocussedAssetText = GetLabelForObject(FocussedDocument);
 	}
 	
 	return FText::Format(LOCTEXT("TookitTooltipFormat", "{0} ({1})"), GetToolTipTextForObject(*WorkspaceObject), FocussedAssetText);
@@ -794,10 +842,10 @@ bool FWorkspaceEditor::ShouldReopenEditorForSavedAsset(const UObject* Asset) con
 
 void FWorkspaceEditor::RemoveEditingObject(UObject* Object)
 {
-	IWorkspaceEditor::RemoveEditingObject(Object);	
-	if (GetFocussedAsset() == Object)
+	IWorkspaceEditor::RemoveEditingObject(Object);
+	if (GetFocussedDocument() == Object)
 	{
-		SetFocussedAsset(nullptr);
+		SetFocussedDocument(nullptr);
 	}
 }
 }
diff --git a/Engine/Plugins/Experimental/Workspace/Source/WorkspaceEditor/Private/WorkspaceEditor.h b/Engine/Plugins/Experimental/Workspace/Source/WorkspaceEditor/Private/WorkspaceEditor.h
--- a/Engine/Plugins/Experimental/Workspace/Source/WorkspaceEditor/Private/WorkspaceEditor.h
+++ b/Engine/Plugins/Experimental/Workspace/Source/WorkspaceEditor/Private/WorkspaceEditor.h
@@ -96,6 +96,7 @@ private:
 	// IWorkspaceEditor interface
 	virtual void OpenAssets(TConstArrayView<FAssetData> InAssets) override;
 	virtual void OpenObjects(TConstArrayView<UObject*> InObjects) override;
+	virtual void GetOpenedAssetsOfClass(TSubclassOf<UObject> InClass, TArray<UObject*>& OutObjects) const override;
 	virtual void GetAssets(TArray<FAssetData>& OutAssets) const override;
 	virtual void CloseObjects(TConstArrayView<UObject*> InObjects) override;
 	virtual void SetDetailsObjects(const TArray<UObject*>& InObjects) override;
@@ -104,10 +105,10 @@ private:
 	virtual bool GetOutlinerSelection(TArray<FWorkspaceOutlinerItemExport>& OutExports) const override;
 	virtual FOnOutlinerSelectionChanged& OnOutlinerSelectionChanged() override;
 	virtual void SetGlobalSelection(FGlobalSelectionId SelectionId, FOnClearGlobalSelection OnClearSelectionDelegate) override;
-	virtual void SetFocussedAsset(const TObjectPtr<UObject> InAsset) override;	
-	virtual const TObjectPtr<UObject> GetFocussedAssetOfClass(const TObjectPtr<UClass> AssetClass) const override;
-	virtual FOnFocussedAssetChanged& OnFocussedAssetChanged() override;
+	virtual const TObjectPtr<UObject> GetFocussedDocumentOfClass(const TObjectPtr<UClass> AssetClass) const override;
+	virtual FOnFocussedDocumentChanged& OnFocussedDocumentChanged() override;
 	virtual TSharedPtr<IDetailsView> GetDetailsView() override;
+	virtual UObject* GetWorkspaceAsset() const override;
 	virtual FString GetPackageName() const override;
 
 	void HandleOutlinerSelectionChanged(TConstArrayView<FWorkspaceOutlinerItemExport> InExports);
@@ -130,6 +131,8 @@ private:
 
 	void RecordDocumentState(const TInstancedStruct<FWorkspaceDocumentState>& InState) const;
 
+	void SetFocussedDocument(const TObjectPtr<UObject> InAsset);
+
 	void NavigateBack();
 	void NavigateForward();
 
@@ -161,7 +164,7 @@ private:
 	int32 SelectionScopeDepth = 0;
 	
 	TArray<FWorkspaceOutlinerItemExport> LastSelectedExports;
-	FOnFocussedAssetChanged OnFocussedAssetChangedDelegate;
+	FOnFocussedDocumentChanged OnFocussedDocumentChangedDelegate;
 	TSharedPtr<FWorkspaceItem> EditorMenuCategory;
 };
 
diff --git a/Engine/Plugins/Experimental/Workspace/Source/WorkspaceEditor/Private/WorkspaceEditorModule.cpp b/Engine/Plugins/Experimental/Workspace/Source/WorkspaceEditor/Private/WorkspaceEditorModule.cpp
--- a/Engine/Plugins/Experimental/Workspace/Source/WorkspaceEditor/Private/WorkspaceEditorModule.cpp
+++ b/Engine/Plugins/Experimental/Workspace/Source/WorkspaceEditor/Private/WorkspaceEditorModule.cpp
@@ -67,6 +67,16 @@ void FWorkspaceEditorModule::UnregisterObjectDocumentType(const FTopLevelAssetPa
 	ObjectDocumentArgs.Remove(InClassPath);
 }
 
+void FWorkspaceEditorModule::RegisterDocumentSubObjectType(const FTopLevelAssetPath& InClassPath, const FDocumentSubObjectArgs& InParams)
+{
+	DocumentSubObjectArgs.Add(InClassPath, InParams);
+}
+
+void FWorkspaceEditorModule::UnregisterDocumentSubObjectType(const FTopLevelAssetPath& InClassPath)
+{
+	DocumentSubObjectArgs.Remove(InClassPath);
+}
+
 FObjectDocumentArgs FWorkspaceEditorModule::CreateGraphDocumentArgs(const FGraphDocumentWidgetArgs& InArgs)
 {
 	FObjectDocumentArgs Args;
@@ -289,6 +299,11 @@ const FObjectDocumentArgs* FWorkspaceEditorModule::FindObjectDocumentType(const
 	return FoundArgs;
 }
 
+const FDocumentSubObjectArgs* FWorkspaceEditorModule::FindDocumentSubObjectType(const UObject* InObject) const
+{
+	return DocumentSubObjectArgs.Find(InObject->GetClass()->GetClassPathName());
+}
+
 TArray<FTopLevelAssetPath> FWorkspaceEditorModule::GetAllowedObjectTypesForArea(FName InSpawnLocation) const
 {
 	return DocumentAreaMap.FindRef(InSpawnLocation).Array();
diff --git a/Engine/Plugins/Experimental/Workspace/Source/WorkspaceEditor/Private/WorkspaceEditorModule.h b/Engine/Plugins/Experimental/Workspace/Source/WorkspaceEditor/Private/WorkspaceEditorModule.h
--- a/Engine/Plugins/Experimental/Workspace/Source/WorkspaceEditor/Private/WorkspaceEditorModule.h
+++ b/Engine/Plugins/Experimental/Workspace/Source/WorkspaceEditor/Private/WorkspaceEditorModule.h
@@ -29,6 +29,8 @@ private:
 	// IWorkspaceModule interface
 	virtual void RegisterObjectDocumentType(const FTopLevelAssetPath& InClassPath, const FObjectDocumentArgs& InParams) override;
 	virtual void UnregisterObjectDocumentType(const FTopLevelAssetPath& InClassPath) override;
+	virtual void RegisterDocumentSubObjectType(const FTopLevelAssetPath& InClassPath, const FDocumentSubObjectArgs& InParams) override;
+	virtual void UnregisterDocumentSubObjectType(const FTopLevelAssetPath& InClassPath) override;
 	virtual FObjectDocumentArgs CreateGraphDocumentArgs(const FGraphDocumentWidgetArgs& InArgs) override;
 	virtual IWorkspaceEditor* OpenWorkspaceForObject(UObject* InObject, EOpenWorkspaceMethod InOpenMethod, const TSubclassOf<UWorkspaceFactory> WorkSpaceFactoryClass) override;	
 	virtual FOnRegisterDetailCustomizations& OnRegisterWorkspaceDetailsCustomization() override;
@@ -41,6 +43,9 @@ private:
 	// Find an existing registered object document type. Note this redirects based on FObjectDocumentArgs::OnRedirectWorkspaceContext
 	const FObjectDocumentArgs* FindObjectDocumentType(const UObject* InObject) const;
 
+	// Find an existing registered document sub-object type.
+	const FDocumentSubObjectArgs* FindDocumentSubObjectType(const UObject* InObject) const;
+
 	// Find the set of allowed object types for the specified spawn location
 	TArray<FTopLevelAssetPath> GetAllowedObjectTypesForArea(FName InSpawnLocation) const;
 
@@ -53,7 +58,9 @@ private:
 	TMap<FTopLevelAssetPath, FObjectDocumentArgs> ObjectDocumentArgs;
 
 	TMap<FName, TSet<FTopLevelAssetPath>> DocumentAreaMap;
-	
+
+	TMap<FTopLevelAssetPath, FDocumentSubObjectArgs> DocumentSubObjectArgs; 
+
 	static TMap<FOutlinerItemDetailsId, TSharedPtr<IWorkspaceOutlinerItemDetails>> OutlinerItemDetails;
 	static TSharedPtr<IWorkspaceOutlinerItemDetails> GetOutlinerItemDetails(const FOutlinerItemDetailsId& DetailId);
 
diff --git a/Engine/Plugins/Experimental/Workspace/Source/WorkspaceEditor/Private/WorkspaceState.h b/Engine/Plugins/Experimental/Workspace/Source/WorkspaceEditor/Private/WorkspaceState.h
--- a/Engine/Plugins/Experimental/Workspace/Source/WorkspaceEditor/Private/WorkspaceState.h
+++ b/Engine/Plugins/Experimental/Workspace/Source/WorkspaceEditor/Private/WorkspaceState.h
@@ -33,10 +33,18 @@ class UWorkspaceState : public UObject
 	// Update soft object paths to any assets we reference
 	void HandleAssetRenamed(const FAssetData& InAssetData, const FString& InOldName);
 
+	// Accessors for user state
+	void SetUserState(FInstancedStruct&& InUserState) { UserState = MoveTemp(InUserState); }
+	const FInstancedStruct& GetUserState() const { return UserState; }
+
 	// Path to the workspace we are persisting, to allow files to be more easily parsed out
 	UPROPERTY()
 	FSoftObjectPath WorkspacePath;
-	
+
+	// User workspace state
+	UPROPERTY()
+	FInstancedStruct UserState;
+
 	// All the workspace documents we are persisting
 	UPROPERTY()
 	TArray<TInstancedStruct<FWorkspaceDocumentState>> DocumentStates;
diff --git a/Engine/Plugins/Experimental/Workspace/Source/WorkspaceEditor/Public/IWorkspaceEditor.h b/Engine/Plugins/Experimental/Workspace/Source/WorkspaceEditor/Public/IWorkspaceEditor.h
--- a/Engine/Plugins/Experimental/Workspace/Source/WorkspaceEditor/Public/IWorkspaceEditor.h
+++ b/Engine/Plugins/Experimental/Workspace/Source/WorkspaceEditor/Public/IWorkspaceEditor.h
@@ -15,7 +15,7 @@ namespace UE::Workspace
 typedef TWeakPtr<SWidget> FGlobalSelectionId;
 using FOnClearGlobalSelection = FSimpleDelegate;
 
-DECLARE_MULTICAST_DELEGATE_OneParam(FOnFocussedAssetChanged, TObjectPtr<UObject>);
+DECLARE_MULTICAST_DELEGATE_OneParam(FOnFocussedDocumentChanged, TObjectPtr<UObject>);
 
 // RAII helper allowing for a multi-widget selection scope within a WorkspaceEditor instance
 struct WORKSPACEEDITOR_API FWorkspaceEditorSelectionScope
@@ -37,6 +37,16 @@ public:
 	// Open the supplied objects for editing within the workspace editor
 	virtual void OpenObjects(TConstArrayView<UObject*> InObjects) = 0;
 
+	// Get the current set of opened (loaded) assets of the specified class
+	virtual void GetOpenedAssetsOfClass(TSubclassOf<UObject> InClass, TArray<UObject*>& OutObjects) const = 0;
+
+	// Get the current set of opened (loaded) assets 
+	void GetOpenedAssets(TArray<UObject*>& OutObjects) const { GetOpenedAssetsOfClass(UObject::StaticClass(), OutObjects); }
+
+	// Get the current set of opened (loaded) assets
+	template<typename AssetClass>
+	void GetOpenedAssets(TArray<UObject*>& OutObjects) const { GetOpenedAssetsOfClass(AssetClass::StaticClass(), OutObjects); }
+
 	// Get the current set of assets in this workspace editor 
 	virtual void GetAssets(TArray<FAssetData>& OutAssets) const = 0;
 
@@ -55,15 +65,24 @@ public:
 	// Set the _current_ global selection (last SWidget with selection set) with delegate to clear it selection on next SetGlobalSelection()
 	virtual void SetGlobalSelection(FGlobalSelectionId SelectionId, FOnClearGlobalSelection OnClearSelectionDelegate) = 0;
 
-	virtual void SetFocussedAsset(const TObjectPtr<UObject> InAsset) = 0;
-	virtual const TObjectPtr<UObject> GetFocussedAssetOfClass(const TObjectPtr<UClass> InClass ) const = 0;
-	
+	// Get the currently focussed document. @return nullptr if the class does not match or no document is focussed
+	virtual const TObjectPtr<UObject> GetFocussedDocumentOfClass(const TObjectPtr<UClass> InClass) const = 0;
+
+	// Get the currently focussed document. @return nullptr if the class does not match or no document is focussed
 	template<typename AssetClass>
-	TObjectPtr<AssetClass> GetFocussedAsset() const { return Cast<AssetClass>(GetFocussedAssetOfClass(AssetClass::StaticClass())); }
-	TObjectPtr<UObject> GetFocussedAsset() const { return GetFocussedAssetOfClass(UObject::StaticClass()); }
+	TObjectPtr<AssetClass> GetFocussedDocument() const
+	{
+		return Cast<AssetClass>(GetFocussedDocumentOfClass(AssetClass::StaticClass()));
+	}
+
+	// Get the currently focussed document. @return nullptr if no document is focussed
+	TObjectPtr<UObject> GetFocussedDocument() const
+	{
+		return GetFocussedDocumentOfClass(UObject::StaticClass());
+	}
 
-	// Multi-cast delegate broadcasted whenever the asset focussed inside of the WorkspaceEditor changes
-	virtual FOnFocussedAssetChanged& OnFocussedAssetChanged() = 0;
+	// Multicast delegate broadcast whenever the document focussed inside of the WorkspaceEditor changes
+	virtual FOnFocussedDocumentChanged& OnFocussedDocumentChanged() = 0;
 
 	// Get the current single selection of the outliner.
 	// @return true if a single selection is active
@@ -76,6 +95,9 @@ public:
 	// Retrieves the common DetailsView widget
 	virtual TSharedPtr<IDetailsView> GetDetailsView() = 0;
 
+	// Returns the workspace asset
+	virtual UObject* GetWorkspaceAsset() const = 0;
+
 	// Returns the name of the package where the workspace is located
 	virtual FString GetPackageName() const = 0;
 };
diff --git a/Engine/Plugins/Experimental/Workspace/Source/WorkspaceEditor/Public/IWorkspaceEditorModule.h b/Engine/Plugins/Experimental/Workspace/Source/WorkspaceEditor/Public/IWorkspaceEditorModule.h
--- a/Engine/Plugins/Experimental/Workspace/Source/WorkspaceEditor/Public/IWorkspaceEditorModule.h
+++ b/Engine/Plugins/Experimental/Workspace/Source/WorkspaceEditor/Public/IWorkspaceEditorModule.h
@@ -81,6 +81,10 @@ using FOnSetDocumentState = TDelegate<void(const FWorkspaceEditorContext&, TShar
 
 using FOnGetDocumentBreadcrumbTrail = TDelegate<void(const FWorkspaceEditorContext&, TArray<TSharedPtr<FWorkspaceBreadcrumb>>&)>;
 
+using FOnGetDocumentForSubObject = TDelegate<UObject*(const UObject*)>;
+
+using FOnPostDocumentOpenedForSubObject = TDelegate<void(const FWorkspaceEditorContext&, TSharedRef<SWidget>, UObject*)>;
+
 // Arguments used to make document widgets for objects
 struct FObjectDocumentArgs
 {
@@ -123,6 +127,16 @@ struct FObjectDocumentArgs
 	FEditorModeID DocumentEditorMode = NAME_None;
 };
 
+// Arguments used to open documents for specific subobject types
+struct FDocumentSubObjectArgs
+{
+	// Delegate called to get a document to open for a document's subobject (e.g. a UEdGraphNode could return its containing UEdGraph)
+	FOnGetDocumentForSubObject OnGetDocumentForSubObject;
+
+	// Delegate called after a document is opened to process the supplied subobject (e.g. focus it).
+	// The widget supplied is the document widget provided via OnMakeDocumentWidget (or a SGraphEditor for FGraphDocumentWidgetArgs)
+	FOnPostDocumentOpenedForSubObject OnPostDocumentOpenedForSubObject;
+};
 
 using FOnGraphSelectionChanged = TDelegate<void(const FWorkspaceEditorContext&, const FGraphPanelSelectionSet&)>;
 
@@ -199,6 +213,12 @@ public:
 	// Unregister a widget factory method to spawn for a particular class
 	virtual void UnregisterObjectDocumentType(const FTopLevelAssetPath& InClassPath) = 0;
 
+	// Register a document subobject - an object that opens in the context of another outer document (e.g. A UEdGraphNode in a UEdGraph)
+	virtual void RegisterDocumentSubObjectType(const FTopLevelAssetPath& InClassPath, const FDocumentSubObjectArgs& InParams) = 0;
+
+	// Unregister a document subobject - an object that opens in the context of another outer document (e.g. A UEdGraphNode in a UEdGraph)
+	virtual void UnregisterDocumentSubObjectType(const FTopLevelAssetPath& InClassPath) = 0;
+
 	// Make the required args for a document widget for a UEdGraph
 	virtual FObjectDocumentArgs CreateGraphDocumentArgs(const FGraphDocumentWidgetArgs& InArgs) = 0;
 
@@ -218,7 +238,7 @@ public:
 	virtual FOnExtendTabs& OnExtendTabs() = 0;
 
 	// Event to allow extending the FToolMenuContext
-	DECLARE_EVENT_TwoParams(IWorkspaceEditorModule, FOnExtendToolMenuContext, const TWeakPtr<IWorkspaceEditor>&, FToolMenuContext&);
+	DECLARE_EVENT_TwoParams(IWorkspaceEditorModule, FOnExtendToolMenuContext, TSharedPtr<IWorkspaceEditor>, FToolMenuContext&);
 	virtual FOnExtendToolMenuContext& OnExtendToolMenuContext() = 0;
 };
 
diff --git a/Engine/Plugins/Experimental/Workspace/Source/WorkspaceEditor/Public/WorkspaceSchema.h b/Engine/Plugins/Experimental/Workspace/Source/WorkspaceEditor/Public/WorkspaceSchema.h
--- a/Engine/Plugins/Experimental/Workspace/Source/WorkspaceEditor/Public/WorkspaceSchema.h
+++ b/Engine/Plugins/Experimental/Workspace/Source/WorkspaceEditor/Public/WorkspaceSchema.h
@@ -5,6 +5,12 @@
 #include "WorkspaceSchema.generated.h"
 
 class UWorkspace;
+struct FInstancedStruct;
+
+namespace UE::Workspace
+{
+	class IWorkspaceEditor;
+}
 
 // Workspace schema to configure workspace assets to specific use cases
 UCLASS(Abstract)
@@ -18,4 +24,10 @@ public:
 
 	// Get the asset types that are supported by this workspace. If this is empty, all assets are assumed to be supported
 	virtual TConstArrayView<FTopLevelAssetPath> GetSupportedAssetClassPaths() const PURE_VIRTUAL(UWorkspaceSchema::GetSupportedAssetClassPaths, return TConstArrayView<FTopLevelAssetPath>(); )
+
+	// Called prior to saving workspace state to populate an instanced struct to hold user-defined persistent workspace state.
+	virtual void OnSaveWorkspaceState(TSharedRef<UE::Workspace::IWorkspaceEditor> InWorkspaceEditor, FInstancedStruct& OutWorkspaceState) const {}
+
+	// Called after loading workspace state. InWorkspaceState is an instanced struct that holds user-defined persistent workspace state. Struct is not guaranteed to be valid.
+	virtual void OnLoadWorkspaceState(TSharedRef<UE::Workspace::IWorkspaceEditor> InWorkspaceEditor, const FInstancedStruct& InWorkspaceState) const {}
 };
\ No newline at end of file
diff --git a/Engine/Plugins/Runtime/RigVM/Source/RigVMDeveloper/Private/RigVMCompiler/RigVMCompiler.cpp b/Engine/Plugins/Runtime/RigVM/Source/RigVMDeveloper/Private/RigVMCompiler/RigVMCompiler.cpp
--- a/Engine/Plugins/Runtime/RigVM/Source/RigVMDeveloper/Private/RigVMCompiler/RigVMCompiler.cpp
+++ b/Engine/Plugins/Runtime/RigVM/Source/RigVMDeveloper/Private/RigVMCompiler/RigVMCompiler.cpp
@@ -760,7 +760,7 @@ bool URigVMCompiler::Compile(const FRigVMCompileSettings& InSettings, TArray<URi
 					{
 						if(ExecutePin->GetLinks().Num() == 0)
 						{
-							static const FString UnlinkedExecuteMessage = TEXT("Node @@ has an unconnected Execute pin.\nThe function might cause unexpected behavior.");
+							static const FString UnlinkedExecuteMessage = TEXT("Node @@ has an unconnected Execute pin. The function might cause unexpected behavior.");
 							Settings.ASTSettings.Report(EMessageSeverity::Error, ModelNode, UnlinkedExecuteMessage);
 							bEncounteredGraphError = true;
 						}
@@ -779,7 +779,7 @@ bool URigVMCompiler::Compile(const FRigVMCompileSettings& InSettings, TArray<URi
 			// for variable let's validate ill formed variable nodes
 			if(URigVMVariableNode* VariableNode = Cast<URigVMVariableNode>(ModelNode))
 			{
-				static const FString IllFormedVariableNodeMessage = TEXT("Variable Node @@ is ill-formed (pin type doesn't match the variable type).\nConsider to recreate the node.");
+				static const FString IllFormedVariableNodeMessage = TEXT("Variable Node @@ is ill-formed (pin type doesn't match the variable type). Consider recreating the node.");
 
 				const FRigVMGraphVariableDescription VariableDescription = VariableNode->GetVariableDescription();
 				const TArray<FRigVMGraphVariableDescription> LocalVariables = VisitedGraph->GetLocalVariables(true);
@@ -1514,7 +1514,6 @@ bool URigVMCompiler::Compile(const FRigVMCompileSettings& InSettings, TArray<URi
 	if (!CurrentCompilationFunction)
 	{
 		CompileTimer.Stop();
-		WorkData.ReportInfof(TEXT("Total Compilation time %f\n"), CompilationTime*1000);
 
 		WorkData.VM->SetVMHash(WorkData.VM->ComputeVMHash());
 		WorkData.VM->GetByteCode().SetPublicContextPathName((WorkData.Context&& WorkData.Context->GetContextPublicDataStruct() != nullptr) ? WorkData.Context->GetContextPublicDataStruct()->GetPathName() : FString());
diff --git a/Engine/Plugins/Runtime/StateTree/Source/StateTreeEditorModule/Private/StandaloneStateTreeEditorHost.cpp b/Engine/Plugins/Runtime/StateTree/Source/StateTreeEditorModule/Private/StandaloneStateTreeEditorHost.cpp
--- a/Engine/Plugins/Runtime/StateTree/Source/StateTreeEditorModule/Private/StandaloneStateTreeEditorHost.cpp
+++ b/Engine/Plugins/Runtime/StateTree/Source/StateTreeEditorModule/Private/StandaloneStateTreeEditorHost.cpp
@@ -27,6 +27,11 @@ FName FStandaloneStateTreeEditorHost::GetCompilerTabName() const
 	return FStateTreeEditor::CompilerResultsTabId;
 }
 
+bool FStandaloneStateTreeEditorHost::ShouldShowCompileButton() const
+{
+	return true;
+}
+
 FSimpleMulticastDelegate& FStandaloneStateTreeEditorHost::OnStateTreeChanged()
 {
 	return OnStateTreeChangedDelegate;
diff --git a/Engine/Plugins/Runtime/StateTree/Source/StateTreeEditorModule/Private/StandaloneStateTreeEditorHost.h b/Engine/Plugins/Runtime/StateTree/Source/StateTreeEditorModule/Private/StandaloneStateTreeEditorHost.h
--- a/Engine/Plugins/Runtime/StateTree/Source/StateTreeEditorModule/Private/StandaloneStateTreeEditorHost.h
+++ b/Engine/Plugins/Runtime/StateTree/Source/StateTreeEditorModule/Private/StandaloneStateTreeEditorHost.h
@@ -16,7 +16,7 @@ public:
 	virtual FSimpleMulticastDelegate& OnStateTreeChanged() override;
 	virtual TSharedPtr<IDetailsView> GetAssetDetailsView() override;
 	virtual TSharedPtr<IDetailsView> GetDetailsView() override;
-
+	virtual bool ShouldShowCompileButton() const override;
 	virtual FName GetCompilerLogName() const override;
 	virtual FName GetCompilerTabName() const override;
 
diff --git a/Engine/Plugins/Runtime/StateTree/Source/StateTreeEditorModule/Private/StateTreeCompilerLog.cpp b/Engine/Plugins/Runtime/StateTree/Source/StateTreeEditorModule/Private/StateTreeCompilerLog.cpp
--- a/Engine/Plugins/Runtime/StateTree/Source/StateTreeEditorModule/Private/StateTreeCompilerLog.cpp
+++ b/Engine/Plugins/Runtime/StateTree/Source/StateTreeEditorModule/Private/StateTreeCompilerLog.cpp
@@ -9,8 +9,9 @@
 
 #define LOCTEXT_NAMESPACE "StateTreeEditor"
 
-void FStateTreeCompilerLog::AppendToLog(IMessageLogListing* LogListing) const
+TArray<TSharedRef<FTokenizedMessage>> FStateTreeCompilerLog::ToTokenizedMessages() const
 {
+	TArray<TSharedRef<FTokenizedMessage>> TokenisedMessages;
 	for (const FStateTreeCompilerLogMessage& StateTreeMessage : Messages)
 	{
 		TSharedRef<FTokenizedMessage> Message = FTokenizedMessage::Create((EMessageSeverity::Type)StateTreeMessage.Severity);
@@ -37,8 +38,15 @@ void FStateTreeCompilerLog::AppendToLog(IMessageLogListing* LogListing) const
 			Message->AddToken(FTextToken::Create(FText::FromString(StateTreeMessage.Message)));
 		}
 
-		LogListing->AddMessage(Message);
+		TokenisedMessages.Add(Message);
 	}
+
+	return TokenisedMessages;
+}
+
+void FStateTreeCompilerLog::AppendToLog(IMessageLogListing* LogListing) const
+{
+	LogListing->AddMessages(ToTokenizedMessages());
 }
 
 void FStateTreeCompilerLog::DumpToLog(const FLogCategoryBase& Category) const
diff --git a/Engine/Plugins/Runtime/StateTree/Source/StateTreeEditorModule/Private/StateTreeEditorMode.cpp b/Engine/Plugins/Runtime/StateTree/Source/StateTreeEditorModule/Private/StateTreeEditorMode.cpp
--- a/Engine/Plugins/Runtime/StateTree/Source/StateTreeEditorModule/Private/StateTreeEditorMode.cpp
+++ b/Engine/Plugins/Runtime/StateTree/Source/StateTreeEditorModule/Private/StateTreeEditorMode.cpp
@@ -339,7 +339,7 @@ void UStateTreeEditorMode::BindToolkitCommands(const TSharedRef<FUICommandList>&
 		FExecuteAction::CreateUObject(this, &UStateTreeEditorMode::Compile),
 		FCanExecuteAction::CreateUObject(this, &UStateTreeEditorMode::CanCompile),
 		FIsActionChecked(),
-		FIsActionButtonVisible::CreateUObject(this, &UStateTreeEditorMode::HasValidStateTree));
+		FIsActionButtonVisible::CreateUObject(this, &UStateTreeEditorMode::IsCompileVisible));
 
 	ToolkitCommands->MapAction(
 		FStateTreeEditorCommands::Get().SaveOnCompile_Never,
@@ -435,6 +435,23 @@ bool UStateTreeEditorMode::CanCompile() const
 	return true;
 }
 
+bool UStateTreeEditorMode::IsCompileVisible() const
+{
+	if(!HasValidStateTree())
+	{
+		return false;
+	}
+
+	if (const UContextObjectStore* ContextObjectStore = GetToolManager()->GetContextObjectStore())
+	{
+		if (const UStateTreeEditorContext* Context = ContextObjectStore->FindContext<UStateTreeEditorContext>())
+		{
+			return Context->EditorHostInterface->ShouldShowCompileButton();
+		}
+	}
+	return true;
+}
+
 bool UStateTreeEditorMode::HasValidStateTree() const
 {
 	return GetStateTree() != nullptr;
diff --git a/Engine/Plugins/Runtime/StateTree/Source/StateTreeEditorModule/Public/IStateTreeEditorHost.h b/Engine/Plugins/Runtime/StateTree/Source/StateTreeEditorModule/Public/IStateTreeEditorHost.h
--- a/Engine/Plugins/Runtime/StateTree/Source/StateTreeEditorModule/Public/IStateTreeEditorHost.h
+++ b/Engine/Plugins/Runtime/StateTree/Source/StateTreeEditorModule/Public/IStateTreeEditorHost.h
@@ -18,7 +18,7 @@ public:
 	virtual ~IStateTreeEditorHost() = default;
 	virtual FName GetCompilerLogName() const = 0;
 	virtual FName GetCompilerTabName() const = 0;
-	
+	virtual bool ShouldShowCompileButton() const = 0;
 	virtual UStateTree* GetStateTree() const = 0;
 	virtual FSimpleMulticastDelegate& OnStateTreeChanged() = 0;
 	virtual TSharedPtr<IDetailsView> GetAssetDetailsView() = 0;
diff --git a/Engine/Plugins/Runtime/StateTree/Source/StateTreeEditorModule/Public/StateTreeCompilerLog.h b/Engine/Plugins/Runtime/StateTree/Source/StateTreeEditorModule/Public/StateTreeCompilerLog.h
--- a/Engine/Plugins/Runtime/StateTree/Source/StateTreeEditorModule/Public/StateTreeCompilerLog.h
+++ b/Engine/Plugins/Runtime/StateTree/Source/StateTreeEditorModule/Public/StateTreeCompilerLog.h
@@ -91,6 +91,9 @@ struct STATETREEEDITORMODULE_API FStateTreeCompilerLog
 		Report(InSeverity, FStateTreeBindableStructDesc(), FString::Printf(Fmt, Args...));
 	}
 
+	/** Generates tokenized messages from the messages in this log */
+	TArray<TSharedRef<FTokenizedMessage>> ToTokenizedMessages() const;
+
 	/** Appends StateTree log to log listing. */
 	void AppendToLog(class IMessageLogListing* LogListing) const;
 
diff --git a/Engine/Plugins/Runtime/StateTree/Source/StateTreeEditorModule/Public/StateTreeEditorMode.h b/Engine/Plugins/Runtime/StateTree/Source/StateTreeEditorModule/Public/StateTreeEditorMode.h
--- a/Engine/Plugins/Runtime/StateTree/Source/StateTreeEditorModule/Public/StateTreeEditorMode.h
+++ b/Engine/Plugins/Runtime/StateTree/Source/StateTreeEditorModule/Public/StateTreeEditorMode.h
@@ -44,6 +44,7 @@ protected:
 
 	void Compile();
 	bool CanCompile() const;
+	bool IsCompileVisible() const;
 
 	bool HasValidStateTree() const;
 	
