AnimNext: fix LOCTEXT key conflict

diff --git a/Engine/Plugins/Experimental/AnimNext/Source/AnimNextEditor/Private/AnimNextEditorModule.cpp b/Engine/Plugins/Experimental/AnimNext/Source/AnimNextEditor/Private/AnimNextEditorModule.cpp
--- a/Engine/Plugins/Experimental/AnimNext/Source/AnimNextEditor/Private/AnimNextEditorModule.cpp
+++ b/Engine/Plugins/Experimental/AnimNext/Source/AnimNextEditor/Private/AnimNextEditorModule.cpp
@@ -220,7 +220,7 @@ void FAnimNextEditorModule::StartupModule()
 				}
 				else
 				{
-					TextBuilder.AppendLine(LOCTEXT("AutoCompileEnabledTooltip", "Auto-compilation is disabled"));
+					TextBuilder.AppendLine(LOCTEXT("AutoCompileDisabledTooltip", "Auto-compilation is disabled"));
 				}
 
 				return TextBuilder.ToText();
