AnimNext: Misc fixes for scene graph

Added ability to use the direct output of a UOL resolve for bindings. Allows use of simple bindings to components on the current actor/entity.
Fixed UOL bindings widget getting closed/refreshed when selecting UOL fragments (programmatic graph generation was causing modification callbacks to be propagated, so we now early-out if compiling).
Temporarily passed a separate interface ptr for the proxy host, as Verse classes do not have UHT-generated interfaces (until we get rid of VNI), so Cast<IInterface> always fails on Verse objects.
Fixed GC not picking up references in editor data CachedExports. This could cause issues on Verse compilation if referencing Verse classes.
Removed dependency on ControlRig
Correctly use component-level physics asset overrides when generating ref poses

diff --git a/Engine/Plugins/Experimental/AnimNext/AnimNext.uplugin b/Engine/Plugins/Experimental/AnimNext/AnimNext.uplugin
--- a/Engine/Plugins/Experimental/AnimNext/AnimNext.uplugin
+++ b/Engine/Plugins/Experimental/AnimNext/AnimNext.uplugin
@@ -18,7 +18,7 @@
 	"Plugins":
 	[
 		{
-			"Name": "ControlRig",
+			"Name": "RigVM",
 			"Enabled": true
 		},
 		{
diff --git a/Engine/Plugins/Experimental/AnimNext/Source/AnimNext/AnimNext.Build.cs b/Engine/Plugins/Experimental/AnimNext/Source/AnimNext/AnimNext.Build.cs
--- a/Engine/Plugins/Experimental/AnimNext/Source/AnimNext/AnimNext.Build.cs
+++ b/Engine/Plugins/Experimental/AnimNext/Source/AnimNext/AnimNext.Build.cs
@@ -12,7 +12,6 @@ namespace UnrealBuildTool.Rules
 					"Core",
 					"CoreUObject",
 					"RigVM",
-					"ControlRig",
 					"Engine",
 					"RenderCore",
 					"HierarchyTableRuntime"
diff --git a/Engine/Plugins/Experimental/AnimNext/Source/AnimNext/Internal/Component/AnimNextWorldSubsystem.h b/Engine/Plugins/Experimental/AnimNext/Source/AnimNext/Internal/Component/AnimNextWorldSubsystem.h
--- a/Engine/Plugins/Experimental/AnimNext/Source/AnimNext/Internal/Component/AnimNextWorldSubsystem.h
+++ b/Engine/Plugins/Experimental/AnimNext/Source/AnimNext/Internal/Component/AnimNextWorldSubsystem.h
@@ -1,4 +1,4 @@
-﻿// Copyright Epic Games, Inc. All Rights Reserved.
+// Copyright Epic Games, Inc. All Rights Reserved.
 
 #pragma once
 
@@ -95,7 +95,7 @@ protected:
 #endif
 	
 	// Register a handle to the subsystem
-	void RegisterHandle(UE::AnimNext::FModuleHandle& InOutHandle, UAnimNextModule* InModule, UObject* InObject, EAnimNextModuleInitMethod InitMethod);
+	void RegisterHandle(UE::AnimNext::FModuleHandle& InOutHandle, UAnimNextModule* InModule, UObject* InObject, IAnimNextVariableProxyHost* InProxyHost, EAnimNextModuleInitMethod InitMethod);
 
 	// Unregister a handle from the subsystem
 	void UnregisterHandle(UE::AnimNext::FModuleHandle& InOutHandle);
diff --git a/Engine/Plugins/Experimental/AnimNext/Source/AnimNext/Internal/Module/AnimNextModuleInstance.h b/Engine/Plugins/Experimental/AnimNext/Source/AnimNext/Internal/Module/AnimNextModuleInstance.h
--- a/Engine/Plugins/Experimental/AnimNext/Source/AnimNext/Internal/Module/AnimNextModuleInstance.h
+++ b/Engine/Plugins/Experimental/AnimNext/Source/AnimNext/Internal/Module/AnimNextModuleInstance.h
@@ -25,6 +25,7 @@ struct FRigUnit_AnimNextInitializeEvent;
 struct FRigUnit_CopyModuleProxyVariables;
 struct FAnimNextNotifyDispatcherComponent;
 struct FAnimNextGraphInstance;
+class IAnimNextVariableProxyHost;
 
 namespace UE::AnimNext
 {
@@ -80,6 +81,7 @@ struct FAnimNextModuleInstance : public FAnimNextDataInterfaceInstance
 		UAnimNextModule* InModule,
 		UObject* InObject,
 		UE::AnimNext::TPool<FAnimNextModuleInstance>* InPool,
+		IAnimNextVariableProxyHost* InProxyHost,
 		EAnimNextModuleInitMethod InInitMethod);
 	ANIMNEXT_API ~FAnimNextModuleInstance();
 
@@ -200,6 +202,10 @@ private:
 	// The pool that this module instance exists in
 	UE::AnimNext::TPool<FAnimNextModuleInstance>* Pool = nullptr;
 
+	// Seperate proxy host ptr - TODO: this only exists because of UHT not parsing Verse native classes right now.
+	// SHould be removed and replaced with a cast of the Cast<IAnimNextVariableProxyHost>(Object) when VNI gets replaced.
+	IAnimNextVariableProxyHost* ProxyHost = nullptr;
+
 	// Copy of the handle that represents this entry to client systems
 	UE::AnimNext::FModuleHandle Handle;
 
diff --git a/Engine/Plugins/Experimental/AnimNext/Source/AnimNext/Internal/Module/ProxyVariablesContext.h b/Engine/Plugins/Experimental/AnimNext/Source/AnimNext/Internal/Module/ProxyVariablesContext.h
--- a/Engine/Plugins/Experimental/AnimNext/Source/AnimNext/Internal/Module/ProxyVariablesContext.h
+++ b/Engine/Plugins/Experimental/AnimNext/Source/AnimNext/Internal/Module/ProxyVariablesContext.h
@@ -1,4 +1,4 @@
-﻿// Copyright Epic Games, Inc. All Rights Reserved.
+// Copyright Epic Games, Inc. All Rights Reserved.
 
 #pragma once
 
@@ -15,7 +15,7 @@ namespace UE::AnimNext
 struct FProxyVariablesContext
 {
 	// Get the default public variables of the module
-	FAnimNextPublicVariablesProxy& GetPublicVariablesProxy() const;
+	ANIMNEXT_API FAnimNextPublicVariablesProxy& GetPublicVariablesProxy() const;
 
 private:
 	explicit FProxyVariablesContext(FAnimNextModuleInstance& InModuleInstance)
diff --git a/Engine/Plugins/Experimental/AnimNext/Source/AnimNext/Internal/Param/RigUnit_AnimNextParameterBase.h b/Engine/Plugins/Experimental/AnimNext/Source/AnimNext/Internal/Param/RigUnit_AnimNextParameterBase.h
--- a/Engine/Plugins/Experimental/AnimNext/Source/AnimNext/Internal/Param/RigUnit_AnimNextParameterBase.h
+++ b/Engine/Plugins/Experimental/AnimNext/Source/AnimNext/Internal/Param/RigUnit_AnimNextParameterBase.h
@@ -1,14 +1,14 @@
-﻿// Copyright Epic Games, Inc. All Rights Reserved.
+// Copyright Epic Games, Inc. All Rights Reserved.
 
 #pragma once
 
 #include "CoreMinimal.h"
-#include "Units/RigUnit.h"
+#include "RigVMCore/RigVMStruct.h"
 #include "RigUnit_AnimNextParameterBase.generated.h"
 
 /** Deprecated stub - remove once content is updated */
 USTRUCT(meta = (ExecuteContext = "FAnimNextExecuteContext"))
-struct FRigUnit_AnimNextParameterBase : public FRigUnit
+struct FRigUnit_AnimNextParameterBase : public FRigVMStruct
 {
 	GENERATED_BODY()
 };
diff --git a/Engine/Plugins/Experimental/AnimNext/Source/AnimNext/Private/Component/AnimNextComponentWorldSubsystem.cpp b/Engine/Plugins/Experimental/AnimNext/Source/AnimNext/Private/Component/AnimNextComponentWorldSubsystem.cpp
--- a/Engine/Plugins/Experimental/AnimNext/Source/AnimNext/Private/Component/AnimNextComponentWorldSubsystem.cpp
+++ b/Engine/Plugins/Experimental/AnimNext/Source/AnimNext/Private/Component/AnimNextComponentWorldSubsystem.cpp
@@ -14,7 +14,7 @@ void UAnimNextComponentWorldSubsystem::Register(UAnimNextComponent* InComponent)
 	using namespace UE::AnimNext;
 	check(IsInGameThread());
 	check(InComponent);
-	RegisterHandle(InComponent->ModuleHandle, InComponent->Module, InComponent, InComponent->InitMethod);
+	RegisterHandle(InComponent->ModuleHandle, InComponent->Module, InComponent, InComponent, InComponent->InitMethod);
 }
 
 void UAnimNextComponentWorldSubsystem::Unregister(UAnimNextComponent* InComponent)
diff --git a/Engine/Plugins/Experimental/AnimNext/Source/AnimNext/Private/Component/AnimNextWorldSubsystem.cpp b/Engine/Plugins/Experimental/AnimNext/Source/AnimNext/Private/Component/AnimNextWorldSubsystem.cpp
--- a/Engine/Plugins/Experimental/AnimNext/Source/AnimNext/Private/Component/AnimNextWorldSubsystem.cpp
+++ b/Engine/Plugins/Experimental/AnimNext/Source/AnimNext/Private/Component/AnimNextWorldSubsystem.cpp
@@ -136,13 +136,13 @@ void UAnimNextWorldSubsystem::FlushPendingActions()
 	}
 }
 
-void UAnimNextWorldSubsystem::RegisterHandle(UE::AnimNext::FModuleHandle& InOutHandle, UAnimNextModule* InModule, UObject* InObject, EAnimNextModuleInitMethod InitMethod)
+void UAnimNextWorldSubsystem::RegisterHandle(UE::AnimNext::FModuleHandle& InOutHandle, UAnimNextModule* InModule, UObject* InObject, IAnimNextVariableProxyHost* InProxyHost, EAnimNextModuleInitMethod InitMethod)
 {
 	using namespace UE::AnimNext;
 	check(IsInGameThread());
 	FWriteScopeLock InstancesLockScope(InstancesLock);
 
-	InOutHandle = Instances.Emplace(InModule, InObject, &Instances, InitMethod);
+	InOutHandle = Instances.Emplace(InModule, InObject, &Instances, InProxyHost, InitMethod);
 	FAnimNextModuleInstance& Instance = Instances.Get(InOutHandle);
 	Instance.Handle = InOutHandle;
 	Instance.Initialize();
diff --git a/Engine/Plugins/Experimental/AnimNext/Source/AnimNext/Private/GenerationTools.cpp b/Engine/Plugins/Experimental/AnimNext/Source/AnimNext/Private/GenerationTools.cpp
--- a/Engine/Plugins/Experimental/AnimNext/Source/AnimNext/Private/GenerationTools.cpp
+++ b/Engine/Plugins/Experimental/AnimNext/Source/AnimNext/Private/GenerationTools.cpp
@@ -210,7 +210,7 @@ struct FCompareBoneIndexType
 		USkeletalMeshComponent::GetRequiredVirtualBones(SkeletalMesh, OutRequiredBones);
 
 		// Add any bones used by physics SkeletalBodySetups
-		const UPhysicsAsset* const PhysicsAsset = SkeletalMesh->GetPhysicsAsset();
+		const UPhysicsAsset* const PhysicsAsset = SkeletalMeshComponent != nullptr ? SkeletalMeshComponent->GetPhysicsAsset() : SkeletalMesh->GetPhysicsAsset();
 		// If we have a PhysicsAsset, we also need to make sure that all the bones used by it are always updated, as its used
 		// by line checks etc. We might also want to kick in the physics, which means having valid bone transforms.
 		if (PhysicsAsset)
diff --git a/Engine/Plugins/Experimental/AnimNext/Source/AnimNext/Private/Graph/RigUnit_AnimNextBeginExecution.cpp b/Engine/Plugins/Experimental/AnimNext/Source/AnimNext/Private/Graph/RigUnit_AnimNextBeginExecution.cpp
--- a/Engine/Plugins/Experimental/AnimNext/Source/AnimNext/Private/Graph/RigUnit_AnimNextBeginExecution.cpp
+++ b/Engine/Plugins/Experimental/AnimNext/Source/AnimNext/Private/Graph/RigUnit_AnimNextBeginExecution.cpp
@@ -1,7 +1,6 @@
 ﻿// Copyright Epic Games, Inc. All Rights Reserved.
 
 #include "Graph/RigUnit_AnimNextBeginExecution.h"
-#include "Units/RigUnitContext.h"
 
 #include UE_INLINE_GENERATED_CPP_BY_NAME(RigUnit_AnimNextBeginExecution)
 
diff --git a/Engine/Plugins/Experimental/AnimNext/Source/AnimNext/Private/Graph/RigUnit_AnimNextRemapPose.cpp b/Engine/Plugins/Experimental/AnimNext/Source/AnimNext/Private/Graph/RigUnit_AnimNextRemapPose.cpp
--- a/Engine/Plugins/Experimental/AnimNext/Source/AnimNext/Private/Graph/RigUnit_AnimNextRemapPose.cpp
+++ b/Engine/Plugins/Experimental/AnimNext/Source/AnimNext/Private/Graph/RigUnit_AnimNextRemapPose.cpp
@@ -2,9 +2,7 @@
 
 #include "RigUnit_AnimNextRemapPose.h"
 
-#include "Units/RigUnitContext.h"
 #include "GenerationTools.h"
-#include "Component/SkinnedMeshComponentExtensions.h"
 #include "Components/SkeletalMeshComponent.h"
 #include "Engine/SkeletalMesh.h"
 
diff --git a/Engine/Plugins/Experimental/AnimNext/Source/AnimNext/Private/Graph/RigUnit_AnimNextWriteSkeletalMeshComponentPose.cpp b/Engine/Plugins/Experimental/AnimNext/Source/AnimNext/Private/Graph/RigUnit_AnimNextWriteSkeletalMeshComponentPose.cpp
--- a/Engine/Plugins/Experimental/AnimNext/Source/AnimNext/Private/Graph/RigUnit_AnimNextWriteSkeletalMeshComponentPose.cpp
+++ b/Engine/Plugins/Experimental/AnimNext/Source/AnimNext/Private/Graph/RigUnit_AnimNextWriteSkeletalMeshComponentPose.cpp
@@ -2,7 +2,6 @@
 
 #include "RigUnit_AnimNextWriteSkeletalMeshComponentPose.h"
 
-#include "Units/RigUnitContext.h"
 #include "AnimNextStats.h"
 #include "GenerationTools.h"
 #include "Component/SkinnedMeshComponentExtensions.h"
diff --git a/Engine/Plugins/Experimental/AnimNext/Source/AnimNext/Private/Graph/RigUnit_MakeReferencePoseFromSkeletalMeshComponent.cpp b/Engine/Plugins/Experimental/AnimNext/Source/AnimNext/Private/Graph/RigUnit_MakeReferencePoseFromSkeletalMeshComponent.cpp
--- a/Engine/Plugins/Experimental/AnimNext/Source/AnimNext/Private/Graph/RigUnit_MakeReferencePoseFromSkeletalMeshComponent.cpp
+++ b/Engine/Plugins/Experimental/AnimNext/Source/AnimNext/Private/Graph/RigUnit_MakeReferencePoseFromSkeletalMeshComponent.cpp
@@ -2,7 +2,6 @@
 
 #include "RigUnit_MakeReferencePoseFromSkeletalMeshComponent.h"
 
-#include "Units/RigUnitContext.h"
 #include "AnimNextStats.h"
 #include "GenerationTools.h"
 #include "Components/SkeletalMeshComponent.h"
diff --git a/Engine/Plugins/Experimental/AnimNext/Source/AnimNext/Private/Module/AnimNextModuleInstance.cpp b/Engine/Plugins/Experimental/AnimNext/Source/AnimNext/Private/Module/AnimNextModuleInstance.cpp
--- a/Engine/Plugins/Experimental/AnimNext/Source/AnimNext/Private/Module/AnimNextModuleInstance.cpp
+++ b/Engine/Plugins/Experimental/AnimNext/Source/AnimNext/Private/Module/AnimNextModuleInstance.cpp
@@ -27,9 +27,11 @@ FAnimNextModuleInstance::FAnimNextModuleInstance(
 		UAnimNextModule* InModule,
 		UObject* InObject,
 		UE::AnimNext::TPool<FAnimNextModuleInstance>* InPool,
+		IAnimNextVariableProxyHost* InProxyHost,
 		EAnimNextModuleInitMethod InInitMethod)
 	: Object(InObject)
 	, Pool(InPool)
+	, ProxyHost(InProxyHost)
 	, RunState(ERunState::None)
 	, InitMethod(InInitMethod)
 {
@@ -354,7 +356,9 @@ void FAnimNextModuleInstance::CopyProxyVariables()
 	// - Hold double-buffered memory handles
 	// - Update the memory handle's ptr to the currently-written double-buffered public variable on write
 	// - Swap the memory handles in ExtendedExecuteContext here
-	if(IAnimNextVariableProxyHost* ProxyHost = Cast<IAnimNextVariableProxyHost>(Object))
+	// 
+//	if(IAnimNextVariableProxyHost* ProxyHost = Cast<IAnimNextVariableProxyHost>(Object))
+	if(ProxyHost)
 	{
 		// Flip the proxy
 		ProxyHost->FlipPublicVariablesProxy(UE::AnimNext::FProxyVariablesContext(*this));
diff --git a/Engine/Plugins/Experimental/AnimNext/Source/AnimNext/Private/Param/RigUnit_AnimNextParameterBeginExecution.cpp b/Engine/Plugins/Experimental/AnimNext/Source/AnimNext/Private/Param/RigUnit_AnimNextParameterBeginExecution.cpp
--- a/Engine/Plugins/Experimental/AnimNext/Source/AnimNext/Private/Param/RigUnit_AnimNextParameterBeginExecution.cpp
+++ b/Engine/Plugins/Experimental/AnimNext/Source/AnimNext/Private/Param/RigUnit_AnimNextParameterBeginExecution.cpp
@@ -1,7 +1,6 @@
 // Copyright Epic Games, Inc. All Rights Reserved.
 
 #include "Param/RigUnit_AnimNextParameterBeginExecution.h"
-#include "Units/RigUnitContext.h"
 
 #include UE_INLINE_GENERATED_CPP_BY_NAME(RigUnit_AnimNextParameterBeginExecution)
 
diff --git a/Engine/Plugins/Experimental/AnimNext/Source/AnimNext/Public/Component/AnimNextComponent.h b/Engine/Plugins/Experimental/AnimNext/Source/AnimNext/Public/Component/AnimNextComponent.h
--- a/Engine/Plugins/Experimental/AnimNext/Source/AnimNext/Public/Component/AnimNextComponent.h
+++ b/Engine/Plugins/Experimental/AnimNext/Source/AnimNext/Public/Component/AnimNextComponent.h
@@ -48,7 +48,7 @@ class UAnimNextComponent : public UActorComponent, public IAnimNextVariableProxy
 	// (Re-)create the public variable proxy
 	void CreatePublicVariablesProxy();
 
-	// (Re-)create the public variable proxy
+	// Destroy the public variable proxy
 	void DestroyPublicVariablesProxy();
 
 public:
diff --git a/Engine/Plugins/Experimental/AnimNext/Source/AnimNextEditor/AnimNextEditor.Build.cs b/Engine/Plugins/Experimental/AnimNext/Source/AnimNextEditor/AnimNextEditor.Build.cs
--- a/Engine/Plugins/Experimental/AnimNext/Source/AnimNextEditor/AnimNextEditor.Build.cs
+++ b/Engine/Plugins/Experimental/AnimNext/Source/AnimNextEditor/AnimNextEditor.Build.cs
@@ -26,9 +26,6 @@ namespace UnrealBuildTool.Rules
 					"RigVM",
 					"RigVMEditor",
 					"RigVMDeveloper",
-					"ControlRig",
-					"ControlRigDeveloper",
-					"ControlRigEditor",
 					"GraphEditor",
 					"ToolWidgets",
 					"ToolMenus",
diff --git a/Engine/Plugins/Experimental/AnimNext/Source/AnimNextEditor/Private/Common/GraphEditorSchemaActions.cpp b/Engine/Plugins/Experimental/AnimNext/Source/AnimNextEditor/Private/Common/GraphEditorSchemaActions.cpp
--- a/Engine/Plugins/Experimental/AnimNext/Source/AnimNextEditor/Private/Common/GraphEditorSchemaActions.cpp
+++ b/Engine/Plugins/Experimental/AnimNext/Source/AnimNextEditor/Private/Common/GraphEditorSchemaActions.cpp
@@ -8,8 +8,6 @@
 #include "Editor.h"
 #include "Module/AnimNextModule_EditorData.h"
 #include "RigVMModel/Nodes/RigVMUnitNode.h"
-#include "Settings/ControlRigSettings.h"
-#include "Units/RigUnit.h"
 #include "EdGraphSchema_K2.h"
 #include "GraphEditor.h"
 #include "Framework/MultiBox/MultiBoxBuilder.h"
@@ -71,24 +69,7 @@ UEdGraphNode* FAnimNextSchemaAction_RigUnit::PerformAction(UEdGraph* ParentGraph
 				Controller->ClearNodeSelection(true);
 				Controller->SelectNode(ModelNode, true, true);
 			}
-
-#if WITH_EDITORONLY_DATA
-
-			// @TODO: switch over to custom settings here
-			const FControlRigSettingsPerPinBool* ExpansionMapPtr = UControlRigEditorSettings::Get()->RigUnitPinExpansion.Find(ModelNode->GetScriptStruct()->GetName());
-			if (ExpansionMapPtr)
-			{
-				const FControlRigSettingsPerPinBool& ExpansionMap = *ExpansionMapPtr;
-
-				for (const TPair<FString, bool>& Pair : ExpansionMap.Values)
-				{
-					FString PinPath = FString::Printf(TEXT("%s.%s"), *ModelNode->GetName(), *Pair.Key);
-					Controller->SetPinExpansion(PinPath, Pair.Value, true);
-				}
-			}
 		}
-#endif
-
 		Controller->CloseUndoBracket();
 	}
 
diff --git a/Engine/Plugins/Experimental/AnimNext/Source/AnimNextEditor/Private/Common/SActionMenu.cpp b/Engine/Plugins/Experimental/AnimNext/Source/AnimNextEditor/Private/Common/SActionMenu.cpp
--- a/Engine/Plugins/Experimental/AnimNext/Source/AnimNextEditor/Private/Common/SActionMenu.cpp
+++ b/Engine/Plugins/Experimental/AnimNext/Source/AnimNextEditor/Private/Common/SActionMenu.cpp
@@ -15,7 +15,6 @@
 #include "RigVMCore/RigVMRegistry.h"
 #include "RigVMModel/RigVMClient.h"
 #include "RigVMModel/RigVMSchema.h"
-#include "Units/RigUnit.h"
 #include "UncookedOnlyUtils.h"
 #include "RigVMBlueprintGeneratedClass.h"
 #include "AnimNextRigVMAssetEditorData.h"
diff --git a/Engine/Plugins/Experimental/AnimNext/Source/AnimNextTestSuite/Private/AnimNextVariablesTest.cpp b/Engine/Plugins/Experimental/AnimNext/Source/AnimNextTestSuite/Private/AnimNextVariablesTest.cpp
--- a/Engine/Plugins/Experimental/AnimNext/Source/AnimNextTestSuite/Private/AnimNextVariablesTest.cpp
+++ b/Engine/Plugins/Experimental/AnimNext/Source/AnimNextTestSuite/Private/AnimNextVariablesTest.cpp
@@ -405,7 +405,7 @@ bool FVariables_UOLBindings::RunTest(const FString& InParameters)
 	Asset->GetRigVMExtendedExecuteContext().SetRuntimeSettings(RuntimeSettings);
 	
 	FAnimNextExecuteContext& AnimNextContext = Asset->GetRigVMExtendedExecuteContext().GetPublicDataSafe<FAnimNextExecuteContext>();
-	FAnimNextModuleInstance Instance(Asset, NewObject<UAnimNextComponent>(), nullptr, EAnimNextModuleInitMethod::None);
+	FAnimNextModuleInstance Instance(Asset, NewObject<UAnimNextComponent>(), nullptr, nullptr, EAnimNextModuleInitMethod::None);
 	FAnimNextModuleContextData ContextData(&Instance);
 	FScopedExecuteContextData ContextDataScope(AnimNextContext, ContextData);
 
diff --git a/Engine/Plugins/Experimental/AnimNext/Source/AnimNextUncookedOnly/AnimNextUncookedOnly.Build.cs b/Engine/Plugins/Experimental/AnimNext/Source/AnimNextUncookedOnly/AnimNextUncookedOnly.Build.cs
--- a/Engine/Plugins/Experimental/AnimNext/Source/AnimNextUncookedOnly/AnimNextUncookedOnly.Build.cs
+++ b/Engine/Plugins/Experimental/AnimNext/Source/AnimNextUncookedOnly/AnimNextUncookedOnly.Build.cs
@@ -11,8 +11,6 @@ namespace UnrealBuildTool.Rules
 				{
 					"RigVM",
 					"RigVMDeveloper",
-					"ControlRig",
-					"ControlRigDeveloper",
 					"AnimNext",
 				}
 			);
diff --git a/Engine/Plugins/Experimental/AnimNext/Source/AnimNextUncookedOnly/Internal/AnimNextRigVMAssetEditorData.h b/Engine/Plugins/Experimental/AnimNext/Source/AnimNextUncookedOnly/Internal/AnimNextRigVMAssetEditorData.h
--- a/Engine/Plugins/Experimental/AnimNext/Source/AnimNextUncookedOnly/Internal/AnimNextRigVMAssetEditorData.h
+++ b/Engine/Plugins/Experimental/AnimNext/Source/AnimNextUncookedOnly/Internal/AnimNextRigVMAssetEditorData.h
@@ -445,6 +445,9 @@ protected:
 	// Handle compiler reporting
 	void HandleReportFromCompiler(EMessageSeverity::Type InSeverity, UObject* InSubject, const FString& InMessage);
 
+	// Support extra references in GC
+	static void AddReferencedObjects(UObject* InThis, FReferenceCollector& Collector);
+
 	// Add a new entry to this asset, taking into account external packaging status
 	void AddEntryInternal(UAnimNextRigVMAssetEntry* InEntry);
 
diff --git a/Engine/Plugins/Experimental/AnimNext/Source/AnimNextUncookedOnly/Internal/Variables/AnimNextUniversalObjectLocatorBindingData.h b/Engine/Plugins/Experimental/AnimNext/Source/AnimNextUncookedOnly/Internal/Variables/AnimNextUniversalObjectLocatorBindingData.h
--- a/Engine/Plugins/Experimental/AnimNext/Source/AnimNextUncookedOnly/Internal/Variables/AnimNextUniversalObjectLocatorBindingData.h
+++ b/Engine/Plugins/Experimental/AnimNext/Source/AnimNextUncookedOnly/Internal/Variables/AnimNextUniversalObjectLocatorBindingData.h
@@ -1,4 +1,4 @@
-﻿// Copyright Epic Games, Inc. All Rights Reserved.
+// Copyright Epic Games, Inc. All Rights Reserved.
 
 #pragma once
 
@@ -12,8 +12,13 @@
 UENUM()
 enum class FAnimNextUniversalObjectLocatorBindingType : uint8
 {
+	// Binding resolves via the result of the UOL (only valid for object type bindings)
+	UOL,
+	// Binding resolves via resolving a property on the UOL's resolved object
 	Property,
+	// Binding resolves via calling a function on the UOL's resolved object
 	Function,
+	// Binding resolves via calling a hoisted function on a BP function library with the UOL's resolved object
 	HoistedFunction,
 };
 
@@ -42,9 +47,10 @@ struct FAnimNextUniversalObjectLocatorBindingData : public FAnimNextVariableBind
 	virtual bool IsValid() const override
 	{
 		const bool bValidLocator = !Locator.IsEmpty();
+		const bool bValidUOL = Type == FAnimNextUniversalObjectLocatorBindingType::UOL;
 		const bool bValidProperty = (Type == FAnimNextUniversalObjectLocatorBindingType::Property && !Property.IsPathToFieldEmpty());
 		const bool bValidFunction = ((Type == FAnimNextUniversalObjectLocatorBindingType::Function || Type == FAnimNextUniversalObjectLocatorBindingType::HoistedFunction) && !Function.IsNull());
-		return bValidLocator && (bValidProperty || bValidFunction);
+		return bValidLocator && (bValidUOL || bValidProperty || bValidFunction);
 	}
 
 #if WITH_EDITORONLY_DATA
diff --git a/Engine/Plugins/Experimental/AnimNext/Source/AnimNextUncookedOnly/Private/AnimNextRigVMAssetEditorData.cpp b/Engine/Plugins/Experimental/AnimNext/Source/AnimNextUncookedOnly/Private/AnimNextRigVMAssetEditorData.cpp
--- a/Engine/Plugins/Experimental/AnimNext/Source/AnimNextUncookedOnly/Private/AnimNextRigVMAssetEditorData.cpp
+++ b/Engine/Plugins/Experimental/AnimNext/Source/AnimNextUncookedOnly/Private/AnimNextRigVMAssetEditorData.cpp
@@ -8,7 +8,7 @@
 #include "AnimNextRigVMAssetSchema.h"
 #include "AnimNextAssetWorkspaceAssetUserData.h"
 #include "AnimNextScopedCompilerResults.h"
-#include "ControlRigDefines.h"
+#include "RigVMPythonUtils.h"
 #include "ExternalPackageHelper.h"
 #include "IAnimNextRigVMGraphInterface.h"
 #include "IWorkspaceEditor.h"
@@ -1028,6 +1028,12 @@ void UAnimNextRigVMAssetEditorData::DecrementVMRecompileBracket()
 
 void UAnimNextRigVMAssetEditorData::HandleModifiedEvent(ERigVMGraphNotifType InNotifType, URigVMGraph* InGraph, UObject* InSubject)
 {
+	// Skip any notifications we get while compiling (they can come from programmatic graph generation)
+	if(bIsCompiling)
+	{
+		return;
+	}
+	
 	bool bNotifForOthersPending = true;
 
 	switch(InNotifType)
@@ -2064,6 +2070,19 @@ void UAnimNextRigVMAssetEditorData::RefreshExternalModels()
 	}
 }
 
+void UAnimNextRigVMAssetEditorData::AddReferencedObjects(UObject* InThis, FReferenceCollector& Collector)
+{
+	Super::AddReferencedObjects(InThis, Collector);
+
+	UAnimNextRigVMAssetEditorData* This = CastChecked<UAnimNextRigVMAssetEditorData>(InThis);
+
+	if (This->CachedExports.IsSet())
+	{
+		// Cached exports may hold references to objects, so make GC aware
+		Collector.AddPropertyReferences(FAnimNextAssetRegistryExports::StaticStruct(), &This->CachedExports.GetValue(), InThis);
+	}
+}
+
 void UAnimNextRigVMAssetEditorData::AddEntryInternal(UAnimNextRigVMAssetEntry* InEntry)
 {
 	// If we are using external packages, dont persist this entry
diff --git a/Engine/Plugins/Experimental/AnimNext/Source/AnimNextUncookedOnly/Private/Variables/UniversalObjectLocatorBindingType.cpp b/Engine/Plugins/Experimental/AnimNext/Source/AnimNextUncookedOnly/Private/Variables/UniversalObjectLocatorBindingType.cpp
--- a/Engine/Plugins/Experimental/AnimNext/Source/AnimNextUncookedOnly/Private/Variables/UniversalObjectLocatorBindingType.cpp
+++ b/Engine/Plugins/Experimental/AnimNext/Source/AnimNextUncookedOnly/Private/Variables/UniversalObjectLocatorBindingType.cpp
@@ -183,6 +183,7 @@ TSharedRef<SWidget> FUniversalObjectLocatorBindingType::CreateEditWidget(const T
 
 		void Construct(const FArguments& InArgs, const TSharedRef<IPropertyHandle>& InPropertyHandle, const FAnimNextParamType& InType)
 		{
+			FilterType = InType;
 			StructProvider = MakeShared<FInstancedStructProvider>(InPropertyHandle);
 			InPropertyHandle->AddChildStructure(StructProvider.ToSharedRef());
 			FieldIterator = MakeUnique<FFieldIterator>(InType);
@@ -192,15 +193,15 @@ TSharedRef<SWidget> FUniversalObjectLocatorBindingType::CreateEditWidget(const T
 			LocatorHandle = InPropertyHandle->GetChildHandle(GET_MEMBER_NAME_CHECKED(FAnimNextUniversalObjectLocatorBindingData, Locator));
 			check(LocatorHandle.IsValid());
 
-			LocatorHandle->SetOnPropertyValueChanged(FSimpleDelegate::CreateLambda([this]()
+			LocatorHandle->SetOnPropertyValueChanged(FSimpleDelegate::CreateLambda([this, InType]()
 			{
 				// Reset property/function/type when setting container
 				PropertyHandle->NotifyPreChange();
-				PropertyHandle->EnumerateRawData([](void* RawData, int32 Index, int32 Num)
+				PropertyHandle->EnumerateRawData([InType](void* RawData, int32 Index, int32 Num)
 				{
 					TInstancedStruct<FAnimNextUniversalObjectLocatorBindingData>* InstancedStruct = static_cast<TInstancedStruct<FAnimNextUniversalObjectLocatorBindingData>*>(RawData);
 					FAnimNextUniversalObjectLocatorBindingData& Data = InstancedStruct->GetMutable();
-					Data.Type = FAnimNextUniversalObjectLocatorBindingType::Property;
+					Data.Type = InType.IsObjectType() ? FAnimNextUniversalObjectLocatorBindingType::UOL : FAnimNextUniversalObjectLocatorBindingType::Property;
 					Data.Property.Reset();
 					Data.Function.Reset();
 					return true;
@@ -369,9 +370,47 @@ TSharedRef<SWidget> FUniversalObjectLocatorBindingType::CreateEditWidget(const T
 		void HandleFieldPicked(UE::PropertyViewer::SPropertyViewer::FHandle InHandle, TArrayView<const FFieldVariant> InFields, ESelectInfo::Type InSelectionType)
 		{
 			FSlateApplication::Get().DismissAllMenus();
-			if(InFields.Num() == 1)
+			if (InFields.Num() == 0 && FilterType.IsObjectType())
+			{
+				// Container picked, so set to 'UOL' if possible
+				if (int32* ContainerIndexPtr = ContainerMap.Find(InHandle))
+				{
+					FContainerInfo& ContainerInfo = CachedContainers[*ContainerIndexPtr];
+					if(ContainerInfo.Struct)
+					{
+						PropertyHandle->NotifyPreChange();
+						PropertyHandle->EnumerateRawData([](void* RawData, int32 Index, int32 Num)
+						{
+							TInstancedStruct<FAnimNextUniversalObjectLocatorBindingData>* InstancedStruct = static_cast<TInstancedStruct<FAnimNextUniversalObjectLocatorBindingData>*>(RawData);
+							FAnimNextUniversalObjectLocatorBindingData& Data = InstancedStruct->GetMutable();
+							Data.Type = FAnimNextUniversalObjectLocatorBindingType::UOL;
+							Data.Property.Reset();
+							Data.Function.Reset();
+							return true;
+						});
+						PropertyHandle->NotifyPostChange(EPropertyChangeType::ValueSet);
+						PropertyHandle->NotifyFinishedChangingProperties();
+					}
+				}
+			}
+			else if(InFields.Num() == 1)
 			{
-				if(const FProperty* Property = InFields[0].Get<FProperty>())
+				if (const UClass* Class = InFields[0].Get<UClass>())
+				{
+					PropertyHandle->NotifyPreChange();
+					PropertyHandle->EnumerateRawData([Class](void* RawData, int32 Index, int32 Num)
+					{
+						TInstancedStruct<FAnimNextUniversalObjectLocatorBindingData>* InstancedStruct = static_cast<TInstancedStruct<FAnimNextUniversalObjectLocatorBindingData>*>(RawData);
+						FAnimNextUniversalObjectLocatorBindingData& Data = InstancedStruct->GetMutable();
+						Data.Type = FAnimNextUniversalObjectLocatorBindingType::UOL;
+						Data.Property.Reset();
+						Data.Function.Reset();
+						return true;
+					});
+					PropertyHandle->NotifyPostChange(EPropertyChangeType::ValueSet);
+					PropertyHandle->NotifyFinishedChangingProperties();
+				}
+				else if(const FProperty* Property = InFields[0].Get<FProperty>())
 				{
 					PropertyHandle->NotifyPreChange();
 					PropertyHandle->EnumerateRawData([Property](void* RawData, int32 Index, int32 Num)
@@ -570,14 +609,15 @@ FText FUniversalObjectLocatorBindingType::GetDisplayText(TConstStructView<FAnimN
 		return false;
 	});
 
-	if(StringBuilder.Len() > 0)
-	{
-		StringBuilder.Append(TEXT("."));
-	}
-
 	switch(LocatorBinding->Type)
 	{
+	case FAnimNextUniversalObjectLocatorBindingType::UOL:
+		break;
 	case FAnimNextUniversalObjectLocatorBindingType::Property:
+		if (StringBuilder.Len() > 0)
+		{
+			StringBuilder.Append(TEXT("."));
+		}
 		if(FProperty* Property = LocatorBinding->Property.Get())
 		{
 			Property->GetFName().AppendString(StringBuilder);
@@ -589,6 +629,10 @@ FText FUniversalObjectLocatorBindingType::GetDisplayText(TConstStructView<FAnimN
 		break;
 	case FAnimNextUniversalObjectLocatorBindingType::Function:
 	case FAnimNextUniversalObjectLocatorBindingType::HoistedFunction:
+		if (StringBuilder.Len() > 0)
+		{
+			StringBuilder.Append(TEXT("."));
+		}
 		if(UFunction* Function = LocatorBinding->Function.Get())
 		{
 			Function->GetFName().AppendString(StringBuilder);
@@ -747,6 +791,8 @@ void FUniversalObjectLocatorBindingType::BuildBindingGraphFragment(const FRigVMC
 			URigVMTemplateNode* GetterNode = nullptr;
 			switch(LocatorBinding.Type)
 			{
+			case FAnimNextUniversalObjectLocatorBindingType::UOL:
+				break;
 			case FAnimNextUniversalObjectLocatorBindingType::Property:
 			{
 				GetterNode = InArgs.Controller->AddTemplateNode(FRigVMDispatch_GetObjectProperty().GetTemplateNotation(), FVector2D(200.0f, YOffset + AccumulatedYOffset), TEXT(""), false);
@@ -866,25 +912,35 @@ void FUniversalObjectLocatorBindingType::BuildBindingGraphFragment(const FRigVMC
 				break;
 			}
 
-			check(GetterNode);
-			URigVMPin* ObjectPin = GetterNode->FindPin(FRigVMDispatch_GetObjectProperty::ObjectName.ToString());
-			if(ObjectPin == nullptr)
-			{
-				InSettings.ReportError(TEXT("Could not find object input pin"));
-				return;
-			}
-			bool bObjectPinLinked = InArgs.Controller->AddLink(ResolvedObjectPin, ObjectPin, false);
-			if(!bObjectPinLinked)
+			URigVMPin* OutputValuePin = nullptr;
+			if(LocatorBinding.Type != FAnimNextUniversalObjectLocatorBindingType::UOL)
 			{
-				InSettings.ReportError(TEXT("Could not link object pins"));
-				return;
+				check(GetterNode);
+				URigVMPin* ObjectPin = GetterNode->FindPin(FRigVMDispatch_GetObjectProperty::ObjectName.ToString());
+				if(ObjectPin == nullptr)
+				{
+					InSettings.ReportError(TEXT("Could not find object input pin"));
+					return;
+				}
+				bool bObjectPinLinked = InArgs.Controller->AddLink(ResolvedObjectPin, ObjectPin, false);
+				if(!bObjectPinLinked)
+				{
+					InSettings.ReportError(TEXT("Could not link object pins"));
+					return;
+				}
+				OutputValuePin = GetterNode->FindPin(FRigVMDispatch_GetObjectProperty::ValueName.ToString());
+				if(OutputValuePin == nullptr)
+				{
+					InSettings.ReportError(TEXT("Could find output value pin"));
+					return;
+				}
 			}
-			URigVMPin* OutputValuePin = GetterNode->FindPin(FRigVMDispatch_GetObjectProperty::ValueName.ToString());
-			if(OutputValuePin == nullptr)
+			else
 			{
-				InSettings.ReportError(TEXT("Could find output value pin"));
-				return;
+				OutputValuePin = ResolvedObjectPin;
 			}
+			check(OutputValuePin);
+
 			URigVMVariableNode* VariableSetNode = InArgs.Controller->AddVariableNode(Input->VariableName, Input->CPPType, Input->CPPTypeObject, false, TEXT(""), FVector2D(400.0f, YOffset + AccumulatedYOffset), TEXT(""), false);
 			if(VariableSetNode == nullptr)
 			{
@@ -910,7 +966,7 @@ void FUniversalObjectLocatorBindingType::BuildBindingGraphFragment(const FRigVMC
 				InSettings.ReportError(TEXT("Could not find Variable Set Value pin"));
 				return;
 			}
-			bool bValuePinLinked = InArgs.Controller->AddLink(OutputValuePin, VariableInputPin, false);
+			bool bValuePinLinked = InArgs.Controller->AddLink(OutputValuePin, VariableInputPin, false, ERigVMPinDirection::Invalid, /*bCreateCastNode*/true);
 			if(!bValuePinLinked)
 			{
 				InSettings.ReportError(TEXT("Could not link Variable Set Value pin"));
diff --git a/Engine/Plugins/Experimental/AnimNextAnimGraph/Source/AnimNextAnimGraph/Private/Graph/RigUnit_AnimNextRunAnimationGraph_v1.cpp b/Engine/Plugins/Experimental/AnimNextAnimGraph/Source/AnimNextAnimGraph/Private/Graph/RigUnit_AnimNextRunAnimationGraph_v1.cpp
--- a/Engine/Plugins/Experimental/AnimNextAnimGraph/Source/AnimNextAnimGraph/Private/Graph/RigUnit_AnimNextRunAnimationGraph_v1.cpp
+++ b/Engine/Plugins/Experimental/AnimNextAnimGraph/Source/AnimNextAnimGraph/Private/Graph/RigUnit_AnimNextRunAnimationGraph_v1.cpp
@@ -2,7 +2,6 @@
 
 #include "Graph/RigUnit_AnimNextRunAnimationGraph_v1.h"
 
-#include "Units/RigUnitContext.h"
 #include "AnimNextAnimGraphStats.h"
 #include "TraitInterfaces/IEvaluate.h"
 #include "TraitInterfaces/IUpdate.h"
diff --git a/Engine/Plugins/Experimental/AnimNextAnimGraph/Source/AnimNextAnimGraph/Private/Graph/RigUnit_AnimNextRunAnimationGraph_v2.cpp b/Engine/Plugins/Experimental/AnimNextAnimGraph/Source/AnimNextAnimGraph/Private/Graph/RigUnit_AnimNextRunAnimationGraph_v2.cpp
--- a/Engine/Plugins/Experimental/AnimNextAnimGraph/Source/AnimNextAnimGraph/Private/Graph/RigUnit_AnimNextRunAnimationGraph_v2.cpp
+++ b/Engine/Plugins/Experimental/AnimNextAnimGraph/Source/AnimNextAnimGraph/Private/Graph/RigUnit_AnimNextRunAnimationGraph_v2.cpp
@@ -3,7 +3,6 @@
 #include "Graph/RigUnit_AnimNextRunAnimationGraph_v2.h"
 
 #include "AnimNextAnimGraphSettings.h"
-#include "Units/RigUnitContext.h"
 #include "AnimNextAnimGraphStats.h"
 #include "TraitInterfaces/IEvaluate.h"
 #include "TraitInterfaces/IUpdate.h"
diff --git a/Engine/Plugins/Experimental/AnimNextAnimGraph/Source/AnimNextAnimGraphUncookedOnly/AnimNextAnimGraphUncookedOnly.Build.cs b/Engine/Plugins/Experimental/AnimNextAnimGraph/Source/AnimNextAnimGraphUncookedOnly/AnimNextAnimGraphUncookedOnly.Build.cs
--- a/Engine/Plugins/Experimental/AnimNextAnimGraph/Source/AnimNextAnimGraphUncookedOnly/AnimNextAnimGraphUncookedOnly.Build.cs
+++ b/Engine/Plugins/Experimental/AnimNextAnimGraph/Source/AnimNextAnimGraphUncookedOnly/AnimNextAnimGraphUncookedOnly.Build.cs
@@ -19,7 +19,6 @@ namespace UnrealBuildTool.Rules
 					"BlueprintGraph",	// For K2 schema
 					"AnimGraph",
 					"RigVM",
-					"ControlRig",
 					"SlateCore",
 					"Slate",
 				}
diff --git a/Engine/Plugins/Experimental/AnimNextAnimGraph/Source/AnimNextAnimGraphUncookedOnly/Internal/Graph/AnimGraphEditorSchemaActions.cpp b/Engine/Plugins/Experimental/AnimNextAnimGraph/Source/AnimNextAnimGraphUncookedOnly/Internal/Graph/AnimGraphEditorSchemaActions.cpp
--- a/Engine/Plugins/Experimental/AnimNextAnimGraph/Source/AnimNextAnimGraphUncookedOnly/Internal/Graph/AnimGraphEditorSchemaActions.cpp
+++ b/Engine/Plugins/Experimental/AnimNextAnimGraph/Source/AnimNextAnimGraphUncookedOnly/Internal/Graph/AnimGraphEditorSchemaActions.cpp
@@ -8,8 +8,6 @@
 #include "Editor.h"
 #include "Module/AnimNextModule_EditorData.h"
 #include "RigVMModel/Nodes/RigVMUnitNode.h"
-#include "Settings/ControlRigSettings.h"
-#include "Units/RigUnit.h"
 #include "EdGraphSchema_K2.h"
 #include "GraphEditor.h"
 #include "Framework/MultiBox/MultiBoxBuilder.h"
diff --git a/Engine/Plugins/Experimental/AnimNextAnimGraph/Source/AnimNextAnimGraphUncookedOnly/Private/Graph/AnimGraph/AnimGraphNode_AnimNextGraph.cpp b/Engine/Plugins/Experimental/AnimNextAnimGraph/Source/AnimNextAnimGraphUncookedOnly/Private/Graph/AnimGraph/AnimGraphNode_AnimNextGraph.cpp
--- a/Engine/Plugins/Experimental/AnimNextAnimGraph/Source/AnimNextAnimGraphUncookedOnly/Private/Graph/AnimGraph/AnimGraphNode_AnimNextGraph.cpp
+++ b/Engine/Plugins/Experimental/AnimNextAnimGraph/Source/AnimNextAnimGraphUncookedOnly/Private/Graph/AnimGraph/AnimGraphNode_AnimNextGraph.cpp
@@ -9,7 +9,6 @@
 #include "ScopedTransaction.h"
 #include "Kismet2/BlueprintEditorUtils.h"
 #include "RigVMBlueprintGeneratedClass.h"
-#include "ControlRigBlueprint.h"
 
 #include UE_INLINE_GENERATED_CPP_BY_NAME(AnimGraphNode_AnimNextGraph)
 
