Hide Internal variables in AnimNext
- Add pass during compilation to generate internal variables

#tests PIE, editor verify variable visibility

diff --git a/Engine/Plugins/Experimental/AnimNext/Source/AnimNextTestSuite/Private/AnimNextDataInterfaceTest.cpp b/Engine/Plugins/Experimental/AnimNext/Source/AnimNextTestSuite/Private/AnimNextDataInterfaceTest.cpp
--- a/Engine/Plugins/Experimental/AnimNext/Source/AnimNextTestSuite/Private/AnimNextDataInterfaceTest.cpp
+++ b/Engine/Plugins/Experimental/AnimNext/Source/AnimNextTestSuite/Private/AnimNextDataInterfaceTest.cpp
@@ -1,4 +1,4 @@
-﻿// Copyright Epic Games, Inc. All Rights Reserved.
+// Copyright Epic Games, Inc. All Rights Reserved.
 
 #include "CoreMinimal.h"
 #include "AnimNextRigVMAssetEditorData.h"
@@ -36,10 +36,11 @@ bool FDataInterfaceCompile::RunTest(const FString& InParameters)
 
 	auto CompileAndTest = [A, B, C, D](TFunctionRef<void()> InTestFunction)
 	{
-		UncookedOnly::FUtils::CompileVariables(A);
-		UncookedOnly::FUtils::CompileVariables(B);
-		UncookedOnly::FUtils::CompileVariables(C);
-		UncookedOnly::FUtils::CompileVariables(D);
+		FRigVMCompileSettings TempSettings = {};
+		UncookedOnly::FUtils::CompileVariables(TempSettings, A);
+		UncookedOnly::FUtils::CompileVariables(TempSettings, B);
+		UncookedOnly::FUtils::CompileVariables(TempSettings, C);
+		UncookedOnly::FUtils::CompileVariables(TempSettings, D);
 
 		InTestFunction();
 	};
diff --git a/Engine/Plugins/Experimental/AnimNext/Source/AnimNextUncookedOnly/Internal/AnimNextRigVMAssetEditorData.h b/Engine/Plugins/Experimental/AnimNext/Source/AnimNextUncookedOnly/Internal/AnimNextRigVMAssetEditorData.h
--- a/Engine/Plugins/Experimental/AnimNext/Source/AnimNextUncookedOnly/Internal/AnimNextRigVMAssetEditorData.h
+++ b/Engine/Plugins/Experimental/AnimNext/Source/AnimNextUncookedOnly/Internal/AnimNextRigVMAssetEditorData.h
@@ -9,6 +9,7 @@
 #include "UncookedOnlyUtils.h"
 #include "Kismet/BlueprintFunctionLibrary.h"
 #include "Entries/AnimNextRigVMAssetEntry.h"
+#include "Variables/AnimNextProgrammaticVariable.h"
 #include "AnimNextRigVMAssetEditorData.generated.h"
 
 enum class ERigVMGraphNotifType : uint8;
@@ -302,6 +303,10 @@ protected:
 	// Override to allow assets to prevent certain entries being created
 	virtual bool CanAddNewEntry(TSubclassOf<UAnimNextRigVMAssetEntry> InClass) const { return true; }
 
+	// Allows this asset to generate variables to be injected at compilation time, separate method to allow programmatic graphs to use these vars
+	// These variables will be regenerated each compile, and are not saved between compiles
+	virtual void GetProgrammaticVariables(const FRigVMCompileSettings& InSettings, TArray<FAnimNextProgrammaticVariable>& OutVariables) {}
+
 	// Allows this asset to generate graphs to be injected at compilation time
 	virtual void GetProgrammaticGraphs(const FRigVMCompileSettings& InSettings, TArray<URigVMGraph*>& OutGraphs) {}
 
diff --git a/Engine/Plugins/Experimental/AnimNext/Source/AnimNextUncookedOnly/Internal/UncookedOnlyUtils.h b/Engine/Plugins/Experimental/AnimNext/Source/AnimNextUncookedOnly/Internal/UncookedOnlyUtils.h
--- a/Engine/Plugins/Experimental/AnimNext/Source/AnimNextUncookedOnly/Internal/UncookedOnlyUtils.h
+++ b/Engine/Plugins/Experimental/AnimNext/Source/AnimNextUncookedOnly/Internal/UncookedOnlyUtils.h
@@ -152,7 +152,7 @@ extern ANIMNEXTUNCOOKEDONLY_API TAutoConsoleVariable<bool> CVarDumpProgrammaticG
 
 struct ANIMNEXTUNCOOKEDONLY_API FUtils
 {
-	static void CompileVariables(UAnimNextRigVMAsset* InAsset);
+	static void CompileVariables(const FRigVMCompileSettings& InSettings, UAnimNextRigVMAsset* InAsset);
 
 	static void CompileVariableBindings(const FRigVMCompileSettings& InSettings, UAnimNextRigVMAsset* InAsset, TArray<URigVMGraph*>& OutGraphs);
 
diff --git a/Engine/Plugins/Experimental/AnimNext/Source/AnimNextUncookedOnly/Private/AnimNextRigVMAssetEditorData.cpp b/Engine/Plugins/Experimental/AnimNext/Source/AnimNextUncookedOnly/Private/AnimNextRigVMAssetEditorData.cpp
--- a/Engine/Plugins/Experimental/AnimNext/Source/AnimNextUncookedOnly/Private/AnimNextRigVMAssetEditorData.cpp
+++ b/Engine/Plugins/Experimental/AnimNext/Source/AnimNextUncookedOnly/Private/AnimNextRigVMAssetEditorData.cpp
@@ -784,7 +784,7 @@ void UAnimNextRigVMAssetEditorData::RecompileVM()
 
 		FUtils::RecreateVM(Asset);
 
-		FUtils::CompileVariables(Asset);
+		FUtils::CompileVariables(Settings, Asset);
 
 		GetProgrammaticGraphs(Settings, ProgrammaticGraphs);
 		for(URigVMGraph* ProgrammaticGraph : ProgrammaticGraphs)
diff --git a/Engine/Plugins/Experimental/AnimNext/Source/AnimNextUncookedOnly/Private/UncookedOnlyUtils.cpp b/Engine/Plugins/Experimental/AnimNext/Source/AnimNextUncookedOnly/Private/UncookedOnlyUtils.cpp
--- a/Engine/Plugins/Experimental/AnimNext/Source/AnimNextUncookedOnly/Private/UncookedOnlyUtils.cpp
+++ b/Engine/Plugins/Experimental/AnimNext/Source/AnimNextUncookedOnly/Private/UncookedOnlyUtils.cpp
@@ -20,6 +20,7 @@
 #include "IWorkspaceEditorModule.h"
 #include "Misc/EnumerateRange.h"
 #include "Module/RigUnit_AnimNextModuleEvents.h"
+#include "Variables/AnimNextProgrammaticVariable.h"
 #include "Variables/IVariableBindingType.h"
 #include "Variables/RigUnit_CopyModuleProxyVariables.h"
 #include "WorkspaceAssetRegistryInfo.h"
@@ -51,7 +52,7 @@ void FUtils::RecreateVM(UAnimNextRigVMAsset* InAsset)
 }
 
 
-void FUtils::CompileVariables(UAnimNextRigVMAsset* InAsset)
+void FUtils::CompileVariables(const FRigVMCompileSettings& InSettings, UAnimNextRigVMAsset* InAsset)
 {
 	check(InAsset);
 
@@ -66,6 +67,10 @@ void FUtils::CompileVariables(UAnimNextRigVMAsset* InAsset)
 
 	UAnimNextDataInterface_EditorData* EditorData = GetEditorData<UAnimNextDataInterface_EditorData>(DataInterface);
 
+	// Gather programmatic variables regenerated each compile
+	TArray<FAnimNextProgrammaticVariable> ProgrammaticVariables;
+	EditorData->GetProgrammaticVariables(InSettings, ProgrammaticVariables);
+
 	struct FStructEntryInfo
 	{
 		const UAnimNextDataInterface* FromDataInterface = nullptr;
@@ -75,6 +80,7 @@ void FUtils::CompileVariables(UAnimNextRigVMAsset* InAsset)
 		EAnimNextExportAccessSpecifier AccessSpecifier = EAnimNextExportAccessSpecifier::Private;
 		bool bAutoBindDataInterfaceToHost = false;
 		TConstArrayView<uint8> Value;
+		EPropertyFlags PropertyFlags = CPF_Edit;
 	};
 
 	// Gather all variables in this asset.
@@ -83,7 +89,7 @@ void FUtils::CompileVariables(UAnimNextRigVMAsset* InAsset)
 	// The topmost value for a data interface 'wins' if a value is to be supplied
 	TMap<FName, int32> EntryInfoIndexMap;
 	TArray<FStructEntryInfo> StructEntryInfos;
-	StructEntryInfos.Reserve(EditorData->Entries.Num());
+	StructEntryInfos.Reserve(EditorData->Entries.Num() + ProgrammaticVariables.Num());
 	int32 NumPublicVariables = 0;
 
 	auto AddVariable = [&Log, &NumPublicVariables, &StructEntryInfos, &EntryInfoIndexMap](const UAnimNextVariableEntry* InVariable, const UAnimNextDataInterfaceEntry* InFromInterfaceEntry, const UAnimNextDataInterface* InFromInterface, bool bInAutoBindInterface)
@@ -158,7 +164,8 @@ void FUtils::CompileVariables(UAnimNextRigVMAsset* InAsset)
 				FAnimNextParamType(Type.GetValueType(), Type.GetContainerType(), Type.GetValueTypeObject()),
 				AccessSpecifier,
 				bInAutoBindInterface,
-				Value
+				Value,
+				CPF_Edit
 			});
 
 			EntryInfoIndexMap.Add(Name, Index);
@@ -203,6 +210,46 @@ void FUtils::CompileVariables(UAnimNextRigVMAsset* InAsset)
 
 	AddDataInterface(DataInterface, nullptr, false, true, AddDataInterface);
 
+	auto AddProgrammaticVariable = [&Log, &StructEntryInfos, &EntryInfoIndexMap](const FAnimNextProgrammaticVariable& ProgrammaticVariable)
+	{
+		const FName Name = ProgrammaticVariable.Name;
+		const FAnimNextParamType& Type = ProgrammaticVariable.Type;
+		if(!Type.IsValid())
+		{
+			Log.Error(FText::Format(LOCTEXT("InvalidProgrammaticVariableTypeFound", "Programmatic Variable '{0}' with invalid type found"), FText::FromName(Name)));
+			return;
+		}
+
+		// Check for type conflicts
+		int32* ExistingIndexPtr = EntryInfoIndexMap.Find(Name);
+		if(ExistingIndexPtr)
+		{
+			Log.Error(FText::Format(LOCTEXT("ConflictingProgrammaticVariableFound", "Programmatic Variable '{0}' already exists, should be created new each compile with no conflicts"), FText::FromName(Name)));
+			return;
+		}
+
+		TConstArrayView<uint8> Value = TConstArrayView<uint8>(ProgrammaticVariable.GetValuePtr(), Type.GetSize());
+
+		int32 Index = StructEntryInfos.Add(
+		{
+			nullptr,
+			nullptr,
+			Name,
+			FAnimNextParamType(Type.GetValueType(), Type.GetContainerType(), Type.GetValueTypeObject()),
+			EAnimNextExportAccessSpecifier::Private,
+			false,
+			Value,
+			CPF_AdvancedDisplay
+		});
+
+		EntryInfoIndexMap.Add(Name, Index);
+	};
+
+	for (const FAnimNextProgrammaticVariable& ProgrammaticVariable : ProgrammaticVariables)
+	{
+		AddProgrammaticVariable(ProgrammaticVariable);
+	}
+
 	// Sort public entries first, then by data interface & then by size, largest first, for better packing
 	static_assert(EAnimNextExportAccessSpecifier::Private < EAnimNextExportAccessSpecifier::Public, "Private must be less than Public as parameters are sorted internally according to this assumption");
 	StructEntryInfos.Sort([](const FStructEntryInfo& InLHS, const FStructEntryInfo& InRHS)
@@ -213,6 +260,12 @@ void FUtils::CompileVariables(UAnimNextRigVMAsset* InAsset)
 		}
 		else if(InLHS.FromDataInterface != InRHS.FromDataInterface)
 		{
+			// If we have a null (Ex: programmatic variables), compare ptrs so that those without interfaces are last
+			if (!InLHS.FromDataInterface || !InRHS.FromDataInterface)
+			{
+				return InLHS.FromDataInterface > InRHS.FromDataInterface;
+			}
+
 			return InLHS.FromDataInterface->GetFName().LexicalLess(InRHS.FromDataInterface->GetFName());
 		}
 		else if(InLHS.Type.GetSize() != InRHS.Type.GetSize())
@@ -239,7 +292,7 @@ void FUtils::CompileVariables(UAnimNextRigVMAsset* InAsset)
 
 		for (TEnumerateRef<const FStructEntryInfo> StructEntryInfo : EnumerateRange(StructEntryInfos))
 		{
-			PropertyDescs.Emplace(StructEntryInfo->Name, StructEntryInfo->Type.ContainerType, StructEntryInfo->Type.ValueType, StructEntryInfo->Type.ValueTypeObject);
+			PropertyDescs.Emplace(StructEntryInfo->Name, StructEntryInfo->Type.ContainerType, StructEntryInfo->Type.ValueType, StructEntryInfo->Type.ValueTypeObject, StructEntryInfo->PropertyFlags);
 			Values.Add(StructEntryInfo->Value);
 
 			// Note: string comparison here because otherwise we would have a circular dependency
@@ -353,7 +406,7 @@ void FUtils::CompileVariableBindingsInternal(const FRigVMCompileSettings& InSett
 
 		TArray<IVariableBindingType::FBindingGraphInput>& Group = BindingGroups.FindOrAdd(BindingType.Get());
 		FRigVMTemplateArgumentType RigVMArg = Variable->GetType().ToRigVMTemplateArgument();
-		Group.Add({ Variable->GetVariableName(), RigVMArg.GetTypeString(), RigVMArg.GetTypeObject(), Binding });
+		Group.Add({ Variable->GetVariableName(), RigVMArg.GetBaseCPPType(), RigVMArg.GetTypeObject(), Binding});
 	}
 
 	const bool bHasBindings = BindingGroups.Num() > 0;
diff --git a/Engine/Plugins/Experimental/AnimNext/Source/AnimNextUncookedOnly/Private/Variables/AnimNextProgrammaticVariable.cpp b/Engine/Plugins/Experimental/AnimNext/Source/AnimNextUncookedOnly/Private/Variables/AnimNextProgrammaticVariable.cpp
new file mode 100644
--- /dev/null
+++ b/Engine/Plugins/Experimental/AnimNext/Source/AnimNextUncookedOnly/Private/Variables/AnimNextProgrammaticVariable.cpp
@@ -0,0 +1,118 @@
+// Copyright Epic Games, Inc. All Rights Reserved.
+
+#pragma once
+
+#include "Variables/AnimNextProgrammaticVariable.h"
+
+#include "Logging/StructuredLog.h"
+#include "RigVMCore/RigVMGraphFunctionDefinition.h"
+#include "RigVMCore/RigVMTemplate.h"
+#include "Variables/IAnimNextRigVMVariableInterface.h"
+
+FAnimNextParamType FAnimNextProgrammaticVariable::GetType() const
+{
+	return Type;
+}
+
+bool FAnimNextProgrammaticVariable::SetType(const FAnimNextParamType& InType, bool bSetupUndoRedo)
+{
+	Type = InType;
+
+	DefaultValue.Reset();
+	DefaultValue.AddProperties({ FPropertyBagPropertyDesc(IAnimNextRigVMVariableInterface::ValueName, Type.GetContainerType(), Type.GetValueType(), Type.GetValueTypeObject()) });
+
+	return true;
+}
+
+FName FAnimNextProgrammaticVariable::GetVariableName() const
+{
+	return Name;
+}
+
+void FAnimNextProgrammaticVariable::SetVariableName(FName InName, bool bSetupUndoRedo)
+{
+	Name = InName;
+}
+
+bool FAnimNextProgrammaticVariable::SetDefaultValue(TConstArrayView<uint8> InValue, bool bSetupUndoRedo)
+{
+	const FPropertyBagPropertyDesc* Desc = DefaultValue.FindPropertyDescByName(IAnimNextRigVMVariableInterface::ValueName);
+	if (Desc == nullptr)
+	{
+		UE_LOGFMT(LogAnimation, Error, "FAnimNextProgrammaticVariable::SetDefaultValue: Could not find default value in property bag");
+		return false;
+	}
+
+	check(Desc->CachedProperty);
+	if (Desc->CachedProperty->GetElementSize() != InValue.Num())
+	{
+		UE_LOGFMT(LogAnimation, Error, "FAnimNextProgrammaticVariable::SetDefaultValue: Mismatched buffer sizes");
+		return false;
+	}
+
+	uint8* DestPtr = Desc->CachedProperty->ContainerPtrToValuePtr<uint8>(DefaultValue.GetMutableValue().GetMemory());
+	const uint8* SrcPtr = InValue.GetData();
+	Desc->CachedProperty->CopyCompleteValue(DestPtr, SrcPtr);
+
+	return true;
+}
+
+bool FAnimNextProgrammaticVariable::SetDefaultValueFromString(const FString& InDefaultValue, bool bSetupUndoRedo)
+{
+	if (DefaultValue.SetValueSerializedString(IAnimNextRigVMVariableInterface::ValueName, InDefaultValue) != EPropertyBagResult::Success)
+	{
+		UE_LOGFMT(LogAnimation, Error, "FAnimNextProgrammaticVariable::FromRigVMGraphFunctionArgument: Could not set value from string");
+		return false;
+	}
+
+	return false;
+}
+
+const FInstancedPropertyBag& FAnimNextProgrammaticVariable::GetPropertyBag() const
+{
+	return DefaultValue;
+}
+
+FInstancedPropertyBag& FAnimNextProgrammaticVariable::GetMutablePropertyBag()
+{
+	return DefaultValue;
+}
+
+bool FAnimNextProgrammaticVariable::GetDefaultValue(const FProperty*& OutProperty, TConstArrayView<uint8>& OutValue) const
+{
+	const FPropertyBagPropertyDesc* Desc = DefaultValue.FindPropertyDescByName(IAnimNextRigVMVariableInterface::ValueName);
+	if (Desc == nullptr)
+	{
+		UE_LOGFMT(LogAnimation, Error, "FAnimNextProgrammaticVariable::FromRigVMGraphFunctionArgument: Could not find default value in property bag");
+		return false;
+	}
+
+	check(Desc->CachedProperty);
+	OutProperty = Desc->CachedProperty;
+	const uint8* ValuePtr = Desc->CachedProperty->ContainerPtrToValuePtr<uint8>(DefaultValue.GetValue().GetMemory());
+	OutValue = TConstArrayView<uint8>(ValuePtr, Desc->CachedProperty->GetElementSize());
+	return true;
+}
+
+const uint8* FAnimNextProgrammaticVariable::GetValuePtr() const
+{
+	const FInstancedPropertyBag& PropertyBag = GetPropertyBag();
+	const UPropertyBag* PropertyBagStruct = PropertyBag.GetPropertyBagStruct();
+	const uint8* Memory = PropertyBag.GetValue().GetMemory();
+	return PropertyBagStruct && Memory ? PropertyBagStruct->GetPropertyDescs()[0].CachedProperty->ContainerPtrToValuePtr<uint8>(Memory) : nullptr;
+}
+
+FAnimNextProgrammaticVariable FAnimNextProgrammaticVariable::FromRigVMGraphFunctionArgument(const FRigVMGraphFunctionArgument& RigVMGraphFunctionArgument)
+{
+	FAnimNextProgrammaticVariable Result = {};
+	Result.Name = RigVMGraphFunctionArgument.Name;
+	Result.SetType(FAnimNextParamType::FromRigVMTemplateArgument(FRigVMTemplateArgumentType(RigVMGraphFunctionArgument.CPPType, RigVMGraphFunctionArgument.CPPTypeObject.Get())));
+	
+	if (RigVMGraphFunctionArgument.DefaultValue.Len() > 0)
+	{
+		// No need to log error on fail, setter will do that
+		Result.SetDefaultValueFromString(RigVMGraphFunctionArgument.DefaultValue);
+	}
+
+	return Result;
+}
diff --git a/Engine/Plugins/Experimental/AnimNext/Source/AnimNextUncookedOnly/Public/Variables/AnimNextProgrammaticVariable.h b/Engine/Plugins/Experimental/AnimNext/Source/AnimNextUncookedOnly/Public/Variables/AnimNextProgrammaticVariable.h
new file mode 100644
--- /dev/null
+++ b/Engine/Plugins/Experimental/AnimNext/Source/AnimNextUncookedOnly/Public/Variables/AnimNextProgrammaticVariable.h
@@ -0,0 +1,64 @@
+// Copyright Epic Games, Inc. All Rights Reserved.
+
+#pragma once
+
+#include "CoreMinimal.h"
+#include "Param/ParamType.h"
+#include "StructUtils/PropertyBag.h"
+#include "Variables/AnimNextVariableBinding.h"
+
+struct FRigVMGraphFunctionArgument;
+
+/**
+ * Struct wrapping a graph variable. Includes default value.
+ */
+struct ANIMNEXTUNCOOKEDONLY_API FAnimNextProgrammaticVariable 
+{
+
+	/** Get arg AnimNext param type */
+	FAnimNextParamType GetType() const;
+
+	/** Set AnimNext param type, also sets propertybag to hold that type */
+	bool SetType(const FAnimNextParamType& InType, bool bSetupUndoRedo = true);
+
+	/** Get VM variable name */
+	FName GetVariableName() const;
+
+	/** Set VM variable name */
+	void SetVariableName(FName InName, bool bSetupUndoRedo = true);
+
+	/** Set variable default value in propertybag */
+	bool SetDefaultValue(TConstArrayView<uint8> InValue, bool bSetupUndoRedo = true);
+
+	/** Set variable default value in propertybag via string */
+	bool SetDefaultValueFromString(const FString& InDefaultValue, bool bSetupUndoRedo = true);
+
+	/** Get inner propertybag storing values */
+	const FInstancedPropertyBag& GetPropertyBag() const;
+
+	/** Get inner propertybag storing values, mutable */
+	FInstancedPropertyBag& GetMutablePropertyBag();
+
+	/** Get default value via property */
+	bool GetDefaultValue(const FProperty*& OutProperty, TConstArrayView<uint8>& OutValue) const;
+
+	/** Access the memory for the internal value */
+	const uint8* GetValuePtr() const;
+
+public:
+
+	/** Name of the variable */
+	FName Name;
+
+	/** The variable's type */
+	FAnimNextParamType Type = FAnimNextParamType::GetType<bool>();
+
+	/** Property bag holding the default value of the variable */
+	FInstancedPropertyBag DefaultValue;
+
+public: 
+
+	/** Construct a parameter type from the passed in FRigVMTemplateArgumentType. */
+	static FAnimNextProgrammaticVariable FromRigVMGraphFunctionArgument(const FRigVMGraphFunctionArgument& RigVMGraphFunctionArgument);
+};
+
diff --git a/Engine/Plugins/Experimental/AnimNextAnimGraph/Source/AnimNextAnimGraphUncookedOnly/Private/Graph/AnimNextAnimationGraph_EditorData.cpp b/Engine/Plugins/Experimental/AnimNextAnimGraph/Source/AnimNextAnimGraphUncookedOnly/Private/Graph/AnimNextAnimationGraph_EditorData.cpp
--- a/Engine/Plugins/Experimental/AnimNextAnimGraph/Source/AnimNextAnimGraphUncookedOnly/Private/Graph/AnimNextAnimationGraph_EditorData.cpp
+++ b/Engine/Plugins/Experimental/AnimNextAnimGraph/Source/AnimNextAnimGraphUncookedOnly/Private/Graph/AnimNextAnimationGraph_EditorData.cpp
@@ -458,7 +458,7 @@ namespace UE::AnimNext::UncookedOnly::Private
 
 			FAnimNextGraphEvaluatorExecuteArgument Argument;
 			Argument.Name = Pin.Name.ToString();
-			Argument.CPPType = TypeArg.GetTypeString();
+			Argument.CPPType = TypeArg.GetBaseCPPType();
 
 			ExecuteDefinition.Arguments.Add(Argument);
 		}
@@ -509,7 +509,7 @@ void UAnimNextAnimationGraph_EditorData::RecompileVM()
 
 		FUtils::RecreateVM(AnimationGraph);
 
-		FUtils::CompileVariables(AnimationGraph);
+		FUtils::CompileVariables(Settings, AnimationGraph);
 
 		AnimationGraph->VMRuntimeSettings = VMRuntimeSettings;
 		AnimationGraph->EntryPoints.Empty();
diff --git a/Engine/Plugins/Experimental/AnimNextStateTree/Source/AnimNextStateTree/Internal/AnimNextStateTreeEditorOnlyTypes.h b/Engine/Plugins/Experimental/AnimNextStateTree/Source/AnimNextStateTree/Internal/AnimNextStateTreeEditorOnlyTypes.h
--- a/Engine/Plugins/Experimental/AnimNextStateTree/Source/AnimNextStateTree/Internal/AnimNextStateTreeEditorOnlyTypes.h
+++ b/Engine/Plugins/Experimental/AnimNextStateTree/Source/AnimNextStateTree/Internal/AnimNextStateTreeEditorOnlyTypes.h
@@ -2,7 +2,10 @@
 
 #pragma once
 
+#if WITH_EDITOR
+
 #include "CoreMinimal.h"
+#include "Variables/AnimNextProgrammaticVariable.h"
 
 class UAnimNextDataInterface_EditorData;
 class URigVMGraph;
@@ -10,6 +13,25 @@ class URigVMGraph;
 struct FRigVMCompileSettings;
 struct FRigVMClient;
 
+/**
+ * Helper struct to contain info needed to compile RigVM Variables
+ */
+struct ANIMNEXTSTATETREE_API FAnimNextStateTreeProgrammaticVariableParams
+{
+	FAnimNextStateTreeProgrammaticVariableParams(UAnimNextDataInterface_EditorData* InEditorData, const FRigVMCompileSettings& InSettings, FRigVMClient& InRigVMClient, TArray<FAnimNextProgrammaticVariable>& InOutVariables)
+		: EditorData(InEditorData)
+		, Settings(InSettings)
+		, RigVMClient(InRigVMClient)
+		, OutVariables(InOutVariables)
+	{
+	}
+
+	UAnimNextDataInterface_EditorData* EditorData;
+	const FRigVMCompileSettings& Settings;
+	FRigVMClient& RigVMClient;
+	TArray<FAnimNextProgrammaticVariable>& OutVariables;
+};
+
 /**
  * Helper struct to contain info needed to compile RigVM Graphs
  */
@@ -28,3 +50,5 @@ struct ANIMNEXTSTATETREE_API FAnimNextStateTreeProgrammaticGraphParams
 	FRigVMClient& RigVMClient;
 	TArray<URigVMGraph*>& OutGraphs;
 };
+
+#endif // WITH_EDITOR
\ No newline at end of file
diff --git a/Engine/Plugins/Experimental/AnimNextStateTree/Source/AnimNextStateTree/Private/Conditions/AnimNextStateTreeRigVMConditionBase.cpp b/Engine/Plugins/Experimental/AnimNextStateTree/Source/AnimNextStateTree/Private/Conditions/AnimNextStateTreeRigVMConditionBase.cpp
--- a/Engine/Plugins/Experimental/AnimNextStateTree/Source/AnimNextStateTree/Private/Conditions/AnimNextStateTreeRigVMConditionBase.cpp
+++ b/Engine/Plugins/Experimental/AnimNextStateTree/Source/AnimNextStateTree/Private/Conditions/AnimNextStateTreeRigVMConditionBase.cpp
@@ -269,49 +269,41 @@ void FAnimNextStateTreeRigVMConditionBase::PostEditNodeChangeChainProperty(const
 	}
 }
 
-void FAnimNextStateTreeRigVMConditionBase::GetProgrammaticGraphs(FAnimNextStateTreeProgrammaticGraphParams& InProgrammaticGraphParams, const UStateTreeState* State, const FStateTreeBindableStructDesc& Desc)
+void FAnimNextStateTreeRigVMConditionBase::GetProgrammaticVariables(FAnimNextStateTreeProgrammaticVariableParams& InProgrammaticVariableParams, const UStateTreeState* State, const FStateTreeBindableStructDesc& Desc)
 {
-	UAnimNextDataInterface_EditorData* EditorData = InProgrammaticGraphParams.EditorData;
-	const FRigVMCompileSettings& Settings = InProgrammaticGraphParams.Settings;
-	FRigVMClient& RigVMClient = InProgrammaticGraphParams.RigVMClient;
-	TArray<URigVMGraph*>& OutGraphs = InProgrammaticGraphParams.OutGraphs;
+	UAnimNextDataInterface_EditorData* EditorData = InProgrammaticVariableParams.EditorData;
+	const FRigVMCompileSettings& Settings = InProgrammaticVariableParams.Settings;
+	FRigVMClient& RigVMClient = InProgrammaticVariableParams.RigVMClient;
+	TArray<FAnimNextProgrammaticVariable>& OutVariables = InProgrammaticVariableParams.OutVariables;
 	
 	StateName = State->Name;
 	NodeId = Desc.ID;
 	InternalResultName = MakeInternalVariableName(ResultName);
 	InternalEventName = MakeInternalEventName(RigVMFunctionHeader.Name);
 
-	URigVMGraph* Graph = NewObject<URigVMGraph>(EditorData, RigVMFunctionHeader.Name, RF_Transient);
-	Graph->SetSchemaClass(UAnimNextAnimationGraphSchema::StaticClass());
-
-	URigVMController* Controller = RigVMClient.GetOrCreateController(Graph);
-	Controller->RemoveNodes(Controller->GetGraph()->GetNodes());
-	
-	// @TODO: Would like to make these variables / procedurally generated ones hidden in AnimNext editor
-	// @TODO: Also need to auto delete these as conditions are removed
 	// Create input variables & result variable
 	{
 		for (const FRigVMGraphFunctionArgument& Argument : RigVMFunctionHeader.Arguments)
 		{
-			FName InternalVariableName = MakeInternalVariableName(Argument.Name);
-			if (Argument.Direction == ERigVMPinDirection::Input)
-			{
-				FRigVMTemplateArgumentType TemplateArgument(Argument.CPPType, Argument.CPPTypeObject.Get());
-				if (!EditorData->FindEntry(InternalVariableName))
-				{
-					EditorData->AddVariable(InternalVariableName, FAnimNextParamType::FromRigVMTemplateArgument(TemplateArgument), Argument.DefaultValue);
-				}
-			}
-			else if (Argument.Direction == ERigVMPinDirection::Output)
-			{
-				FRigVMTemplateArgumentType TemplateArgument(Argument.CPPType, Argument.CPPTypeObject.Get());
-				if (!EditorData->FindEntry(InternalVariableName))
-				{
-					EditorData->AddVariable(InternalVariableName, FAnimNextParamType::FromRigVMTemplateArgument(TemplateArgument), Argument.DefaultValue);
-				}
-			}
+			FRigVMGraphFunctionArgument InternallyNamedArgument = Argument;
+			InternallyNamedArgument.Name = MakeInternalVariableName(Argument.Name);
+			OutVariables.Add(FAnimNextProgrammaticVariable::FromRigVMGraphFunctionArgument(InternallyNamedArgument));
 		}
 	}
+}
+
+void FAnimNextStateTreeRigVMConditionBase::GetProgrammaticGraphs(FAnimNextStateTreeProgrammaticGraphParams& InProgrammaticGraphParams, const UStateTreeState* State, const FStateTreeBindableStructDesc& Desc)
+{
+	UAnimNextDataInterface_EditorData* EditorData = InProgrammaticGraphParams.EditorData;
+	const FRigVMCompileSettings& Settings = InProgrammaticGraphParams.Settings;
+	FRigVMClient& RigVMClient = InProgrammaticGraphParams.RigVMClient;
+	TArray<URigVMGraph*>& OutGraphs = InProgrammaticGraphParams.OutGraphs;
+
+	URigVMGraph* Graph = NewObject<URigVMGraph>(EditorData, RigVMFunctionHeader.Name, RF_Transient);
+	Graph->SetSchemaClass(UAnimNextAnimationGraphSchema::StaticClass());
+
+	URigVMController* Controller = RigVMClient.GetOrCreateController(Graph);
+	Controller->RemoveNodes(Controller->GetGraph()->GetNodes());
 
 	// Create RigVM graph for function
 	{
diff --git a/Engine/Plugins/Experimental/AnimNextStateTree/Source/AnimNextStateTree/Public/Conditions/AnimNextStateTreeRigVMConditionBase.h b/Engine/Plugins/Experimental/AnimNextStateTree/Source/AnimNextStateTree/Public/Conditions/AnimNextStateTreeRigVMConditionBase.h
--- a/Engine/Plugins/Experimental/AnimNextStateTree/Source/AnimNextStateTree/Public/Conditions/AnimNextStateTreeRigVMConditionBase.h
+++ b/Engine/Plugins/Experimental/AnimNextStateTree/Source/AnimNextStateTree/Public/Conditions/AnimNextStateTreeRigVMConditionBase.h
@@ -18,6 +18,7 @@ class UAnimNextDataInterface_EditorData;
 class URigVMGraph;
 class UStateTreeState;
 
+struct FAnimNextStateTreeProgrammaticVariableParams;
 struct FAnimNextStateTreeProgrammaticGraphParams;
 struct FStateTreeBindableStructDesc;
 struct FStateTreeExecutionContext;
@@ -90,6 +91,9 @@ public:
 	//~ End FStateTreeNodeBase Interface
 	
 #if WITH_EDITOR
+	/** Adds condition variables to parent StateTree RigVM */
+	virtual void GetProgrammaticVariables(FAnimNextStateTreeProgrammaticVariableParams& InProgrammaticVariableParams, const UStateTreeState* State, const FStateTreeBindableStructDesc& Desc);
+
 	/** Adds condition graphs to parent StateTree RigVM */
 	virtual void GetProgrammaticGraphs(FAnimNextStateTreeProgrammaticGraphParams& InProgrammaticGraphParams, const UStateTreeState* State, const FStateTreeBindableStructDesc& Desc);
 #endif
diff --git a/Engine/Plugins/Experimental/AnimNextStateTree/Source/AnimNextStateTreeUncookedOnly/Internal/AnimNextStateTree_EditorData.cpp b/Engine/Plugins/Experimental/AnimNextStateTree/Source/AnimNextStateTreeUncookedOnly/Internal/AnimNextStateTree_EditorData.cpp
--- a/Engine/Plugins/Experimental/AnimNextStateTree/Source/AnimNextStateTreeUncookedOnly/Internal/AnimNextStateTree_EditorData.cpp
+++ b/Engine/Plugins/Experimental/AnimNextStateTree/Source/AnimNextStateTreeUncookedOnly/Internal/AnimNextStateTree_EditorData.cpp
@@ -55,6 +55,51 @@ TConstArrayView<TSubclassOf<UAnimNextRigVMAssetEntry>> UAnimNextStateTree_Editor
 	return Classes;
 }
 
+void UAnimNextStateTree_EditorData::GetProgrammaticVariables(const FRigVMCompileSettings& InSettings, TArray<FAnimNextProgrammaticVariable>& OutVariables)
+{
+	if (UAnimNextStateTree* AnimStateTree = UE::AnimNext::UncookedOnly::FUtils::GetAsset<UAnimNextStateTree>(this))
+	{
+		// Give child nodes a chance to add compile-time only variables
+		FAnimNextStateTreeProgrammaticVariableParams ProgrammaticVariableParams(this, InSettings, RigVMClient, OutVariables);
+
+		// Populate UID info on our nodes
+		auto MakeNodeVariables = [&](UStateTreeState& State, UStateTreeState* ParentState)
+		{
+			for (const FStateTreeEditorNode& Node : State.EnterConditions)
+			{
+				// Manually check for each child node type, UStruct's do not reflect interfaces so we can't use an interface here.
+				if (const FAnimNextStateTreeRigVMConditionBase* Condition = Node.Node.GetPtr<FAnimNextStateTreeRigVMConditionBase>())
+				{
+					FStateTreeBindableStructDesc Desc;
+					Desc.StatePath = State.GetPath() + TEXT("/EnterConditions");
+					Desc.Struct = Condition->GetInstanceDataType();
+					Desc.Name = Node.GetName();
+					Desc.ID = Node.ID;
+					Desc.DataSource = EStateTreeBindableStructSource::Condition;
+
+					// Graph gen isn't const since we need to cache generated names in the node
+					const_cast<FAnimNextStateTreeRigVMConditionBase*>(Condition)->GetProgrammaticVariables(ProgrammaticVariableParams, &State, Desc);
+				}
+			}
+
+			// @TODO: Add RigVM graph task support
+			//for (const FStateTreeEditorNode& Node : State.Tasks)
+			//{
+			//}
+
+			return EStateTreeVisitor::Continue;
+		};
+
+		if (AnimStateTree->StateTree)
+		{
+			if (UStateTreeEditorData* StateTreeEditorData = Cast<UStateTreeEditorData>(AnimStateTree->StateTree->EditorData))
+			{
+				StateTreeEditorData->VisitHierarchy(MakeNodeVariables);
+			}
+		}
+	}
+}
+
 void UAnimNextStateTree_EditorData::GetProgrammaticGraphs(const FRigVMCompileSettings& InSettings, TArray<URigVMGraph*>& OutGraphs)
 {
 	if (UAnimNextStateTree* AnimStateTree = UE::AnimNext::UncookedOnly::FUtils::GetAsset<UAnimNextStateTree>(this))
diff --git a/Engine/Plugins/Experimental/AnimNextStateTree/Source/AnimNextStateTreeUncookedOnly/Internal/AnimNextStateTree_EditorData.h b/Engine/Plugins/Experimental/AnimNextStateTree/Source/AnimNextStateTreeUncookedOnly/Internal/AnimNextStateTree_EditorData.h
--- a/Engine/Plugins/Experimental/AnimNextStateTree/Source/AnimNextStateTreeUncookedOnly/Internal/AnimNextStateTree_EditorData.h
+++ b/Engine/Plugins/Experimental/AnimNextStateTree/Source/AnimNextStateTreeUncookedOnly/Internal/AnimNextStateTree_EditorData.h
@@ -21,6 +21,9 @@ protected:
 	// UAnimNextRigVMAssetEditorData interface
 	virtual TConstArrayView<TSubclassOf<UAnimNextRigVMAssetEntry>> GetEntryClasses() const override;
 
+	// Allows this asset to generate variables to be injected at compilation time, separate method to allow programmatic graphs to use these vars
+	virtual void GetProgrammaticVariables(const FRigVMCompileSettings& InSettings, TArray<FAnimNextProgrammaticVariable>& OutVariables) override;
+
 	// Allows this asset to generate graphs to be injected at compilation time
 	virtual void GetProgrammaticGraphs(const FRigVMCompileSettings& InSettings, TArray<URigVMGraph*>& OutGraphs) override;
 };
\ No newline at end of file
