Add Search Support to AnimNext

#tests editor

diff --git a/Engine/Plugins/Experimental/AnimNext/Source/AnimNextEditor/Private/AnimNextEditorModule.cpp b/Engine/Plugins/Experimental/AnimNext/Source/AnimNextEditor/Private/AnimNextEditorModule.cpp
--- a/Engine/Plugins/Experimental/AnimNext/Source/AnimNextEditor/Private/AnimNextEditorModule.cpp
+++ b/Engine/Plugins/Experimental/AnimNext/Source/AnimNextEditor/Private/AnimNextEditorModule.cpp
@@ -10,12 +10,14 @@
 #include "SSimpleButton.h"
 #include "UncookedOnlyUtils.h"
 #include "Editor/RigVMEditorTools.h"
+#include "FindInAnimNextRigVMAsset.h"
 #include "Framework/Application/SlateApplication.h"
 #include "Module/AnimNextModule.h"
 #include "Common/AnimNextEdGraphNodeCustomization.h"
 #include "Module/AnimNextModule_EditorData.h"
 #include "AnimNextEdGraphNode.h"
 #include "Common/AnimNextCompilerResultsTabSummoner.h"
+#include "Common/AnimNextFindTabSummoner.h"
 #include "HAL/PlatformApplicationMisc.h"
 #include "Variables/VariableCustomization.h"
 #include "Variables/VariableBindingPropertyCustomization.h"
@@ -41,6 +43,7 @@
 #include "Common/AnimNextGraphItemDetails.h"
 #include "Common/AnimNextCollapseNodeItemDetails.h"
 #include "Common/AnimNextFunctionItemDetails.h"
+#include "Common/AnimNextRigVMAssetCommands.h"
 #include "Param/AnimNextActorLocatorEditor.h"
 #include "IWorkspaceEditor.h"
 #include "RigVMCommands.h"
@@ -75,6 +78,7 @@ void FAnimNextEditorModule::StartupModule()
 {
 	FVariableOverrideCommands::Register();
 	FRigVMCommands::Register();
+	FAnimNextRigVMAssetCommands::Register();
 
 	// Register settings for user editing
 	ISettingsModule& SettingsModule = FModuleManager::Get().LoadModuleChecked<ISettingsModule>("Settings");
@@ -111,6 +115,10 @@ void FAnimNextEditorModule::StartupModule()
 			TabFactories.RegisterFactory(CompilerResultsTabSummoner);
 			CompilerResultsTabSummoner->RegisterTabSpawner(InTabManager, nullptr);
 
+			TSharedRef<FAnimNextFindTabSummoner> FindTabSummoner = MakeShared<FAnimNextFindTabSummoner>(InEditorPtr);
+			TabFactories.RegisterFactory(FindTabSummoner);
+			FindTabSummoner->RegisterTabSpawner(InTabManager, nullptr);
+
 			TSharedRef<FAnimNextVariablesTabSummoner> VariablesTabSummoner = MakeShared<FAnimNextVariablesTabSummoner>(InEditorPtr);
 			TabFactories.RegisterFactory(VariablesTabSummoner);
 			VariablesTabSummoner->RegisterTabSpawner(InTabManager, nullptr);
@@ -121,6 +129,9 @@ void FAnimNextEditorModule::StartupModule()
 		FTabManager::FTab CompilerResultsTab(FTabId(CompilerResultsTabName), ETabState::ClosedTab);
 		InLayoutExtender.ExtendLayout(FTabId(Workspace::WorkspaceTabs::BottomMiddleDocumentArea), ELayoutExtensionPosition::After, CompilerResultsTab);
 
+		FTabManager::FTab FindTab(FTabId(FindTabName), ETabState::ClosedTab);
+		InLayoutExtender.ExtendLayout(FTabId(Workspace::WorkspaceTabs::BottomMiddleDocumentArea), ELayoutExtensionPosition::After, FindTab);
+
 		FTabManager::FTab VariablesTab(FTabId(VariablesTabName), ETabState::OpenedTab);
 		InLayoutExtender.ExtendLayout(FTabId(Workspace::WorkspaceTabs::BottomLeftDocumentArea), ELayoutExtensionPosition::After, VariablesTab);
 	});
@@ -982,6 +993,31 @@ void FAnimNextEditorModule::RegisterWorkspaceDocumentTypes(Workspace::IWorkspace
 						}
 					}
 				});
+				
+				// Register any general AnimNext commands on document make
+				if (TSharedPtr<Workspace::IWorkspaceEditor> WorkspaceEditor = WeakWorkspaceEditor.Pin())
+				{
+					const TSharedRef<FUICommandList>& CommandList = WorkspaceEditor->GetToolkitCommands();
+
+					{
+						const FAnimNextRigVMAssetCommands& Commands = FAnimNextRigVMAssetCommands::Get();
+
+						auto TryFindAnimNextRigVMAsset = [WeakWorkspaceEditor]()
+						{
+							if (TSharedPtr<Workspace::IWorkspaceEditor> WorkspaceEditorPinned = WeakWorkspaceEditor.Pin())
+							{
+								if (TSharedPtr<SDockTab> Tab = WorkspaceEditorPinned->GetTabManager()->TryInvokeTab(FTabId(FindTabName)))
+								{
+									TSharedRef<SFindInAnimNextRigVMAsset> FindInAnimNextRigVMAssetWidget = StaticCastSharedRef<SFindInAnimNextRigVMAsset>(Tab->GetContent());
+									FindInAnimNextRigVMAssetWidget->FocusForUse();
+								}
+							}
+						};
+
+						CommandList->MapAction(Commands.FindInAnimNextRigVMAsset,
+							FExecuteAction::CreateLambda(TryFindAnimNextRigVMAsset));
+					}
+				}
 			}
 		}
 
@@ -1180,4 +1216,4 @@ const FAssetCompilationHandlerFactoryDelegate* FAnimNextEditorModule::FindAssetC
 
 IMPLEMENT_MODULE(UE::AnimNext::Editor::FAnimNextEditorModule, AnimNextEditor);
 
-#undef LOCTEXT_NAMESPACE
\ No newline at end of file
+#undef LOCTEXT_NAMESPACE
diff --git a/Engine/Plugins/Experimental/AnimNext/Source/AnimNextEditor/Private/Common/AnimNextFindTabSummoner.cpp b/Engine/Plugins/Experimental/AnimNext/Source/AnimNextEditor/Private/Common/AnimNextFindTabSummoner.cpp
new file mode 100644
--- /dev/null
+++ b/Engine/Plugins/Experimental/AnimNext/Source/AnimNextEditor/Private/Common/AnimNextFindTabSummoner.cpp
@@ -0,0 +1,41 @@
+// Copyright Epic Games, Inc. All Rights Reserved.
+
+#include "Common/AnimNextFindTabSummoner.h"
+
+#include "AnimNextRigVMAsset.h"
+#include "FindInAnimNextRigVMAsset.h"
+#include "IAnimNextEditorModule.h"
+#include "IWorkspaceEditor.h"
+#include "MessageLogModule.h"
+#include "Modules/ModuleManager.h"
+
+#define LOCTEXT_NAMESPACE "WorkspaceTabSummoner"
+
+namespace UE::AnimNext::Editor
+{
+
+FAnimNextFindTabSummoner::FAnimNextFindTabSummoner(TSharedPtr<UE::Workspace::IWorkspaceEditor> InHostingApp)
+	: FWorkflowTabFactory(UE::AnimNext::Editor::FindTabName, StaticCastSharedPtr<FAssetEditorToolkit>(InHostingApp))
+{
+	TabLabel = LOCTEXT("AnimNextFindTabLabel", "Find");
+	TabIcon = FSlateIcon(FAppStyle::GetAppStyleSetName(), "Kismet.Tabs.FindResults");
+	ViewMenuDescription = LOCTEXT("AnimNextFindTabMenuDescription", "Find");
+	ViewMenuTooltip = LOCTEXT("AnimNextFindTabToolTip", "Search contents of currently selected AnimNext Asset.");
+	bIsSingleton = true;
+
+	AnimNextFindResultsWidget = SNew(SFindInAnimNextRigVMAsset, InHostingApp);
+}
+
+TSharedRef<SWidget> FAnimNextFindTabSummoner::CreateTabBody(const FWorkflowTabSpawnInfo& Info) const
+{
+	return AnimNextFindResultsWidget.ToSharedRef();
+}
+
+FText FAnimNextFindTabSummoner::GetTabToolTipText(const FWorkflowTabSpawnInfo& Info) const
+{
+	return ViewMenuTooltip;
+}
+
+} // end namespace UE::AnimNext::Editor
+
+#undef LOCTEXT_NAMESPACE
diff --git a/Engine/Plugins/Experimental/AnimNext/Source/AnimNextEditor/Private/Common/AnimNextFindTabSummoner.h b/Engine/Plugins/Experimental/AnimNext/Source/AnimNextEditor/Private/Common/AnimNextFindTabSummoner.h
new file mode 100644
--- /dev/null
+++ b/Engine/Plugins/Experimental/AnimNext/Source/AnimNextEditor/Private/Common/AnimNextFindTabSummoner.h
@@ -0,0 +1,34 @@
+// Copyright Epic Games, Inc. All Rights Reserved.
+
+#pragma once
+
+#include "WorkflowOrientedApp/WorkflowTabFactory.h"
+#include "IMessageLogListing.h"
+
+namespace UE::Workspace
+{
+	class IWorkspaceEditor;
+}	// end namespace UE::Workspace
+
+namespace UE::AnimNext::Editor
+{
+	class SFindInAnimNextRigVMAsset;
+}	// end namespace UE::AnimNext::Editor
+
+namespace UE::AnimNext::Editor
+{
+struct FAnimNextFindTabSummoner : public FWorkflowTabFactory
+{
+public:
+	FAnimNextFindTabSummoner(TSharedPtr<UE::Workspace::IWorkspaceEditor> InHostingApp);
+
+private:
+	// FWorkflowTabFactory interface
+	virtual TSharedRef<SWidget> CreateTabBody(const FWorkflowTabSpawnInfo& Info) const override;
+	virtual FText GetTabToolTipText(const FWorkflowTabSpawnInfo& Info) const override;
+
+	// The widget this tab spawner wraps
+	TSharedPtr<SFindInAnimNextRigVMAsset> AnimNextFindResultsWidget;
+};
+
+} // end namespace UE::AnimNext::Editor
diff --git a/Engine/Plugins/Experimental/AnimNext/Source/AnimNextEditor/Private/Common/AnimNextRigVMAssetCommands.cpp b/Engine/Plugins/Experimental/AnimNext/Source/AnimNextEditor/Private/Common/AnimNextRigVMAssetCommands.cpp
new file mode 100644
--- /dev/null
+++ b/Engine/Plugins/Experimental/AnimNext/Source/AnimNextEditor/Private/Common/AnimNextRigVMAssetCommands.cpp
@@ -0,0 +1,22 @@
+﻿// Copyright Epic Games, Inc. All Rights Reserved.
+
+#include "Common/AnimNextRigVMAssetCommands.h"
+
+#define LOCTEXT_NAMESPACE "AnimNextRigVMAssetCommands"
+
+namespace UE::AnimNext
+{
+
+FAnimNextRigVMAssetCommands::FAnimNextRigVMAssetCommands()
+	: TCommands<FAnimNextRigVMAssetCommands>("AnimNextRigVMAssetCommand", LOCTEXT("AnimNextRigVMAssetCommands", "Anim Next Asset Commands"), NAME_None, "AnimNextStyle")
+{
+}
+
+void FAnimNextRigVMAssetCommands::RegisterCommands()
+{
+	UI_COMMAND(FindInAnimNextRigVMAsset, "FindInAnimNextRigVMAsset", "Search the current AnimNext Asset.", EUserInterfaceActionType::Button, FInputChord(EModifierKey::Control, EKeys::F));
+}
+
+}
+
+#undef LOCTEXT_NAMESPACE
diff --git a/Engine/Plugins/Experimental/AnimNext/Source/AnimNextEditor/Private/Common/AnimNextRigVMAssetCommands.h b/Engine/Plugins/Experimental/AnimNext/Source/AnimNextEditor/Private/Common/AnimNextRigVMAssetCommands.h
new file mode 100644
--- /dev/null
+++ b/Engine/Plugins/Experimental/AnimNext/Source/AnimNextEditor/Private/Common/AnimNextRigVMAssetCommands.h
@@ -0,0 +1,21 @@
+﻿// Copyright Epic Games, Inc. All Rights Reserved.
+
+#pragma once
+
+#include "Framework/Commands/Commands.h"
+
+namespace UE::AnimNext
+{
+
+// The Commands for the variable overrides menu
+class FAnimNextRigVMAssetCommands : public TCommands<FAnimNextRigVMAssetCommands>
+{
+public:
+	FAnimNextRigVMAssetCommands();
+
+	virtual void RegisterCommands() override;
+
+	TSharedPtr<FUICommandInfo> FindInAnimNextRigVMAsset;
+};
+
+}
\ No newline at end of file
diff --git a/Engine/Plugins/Experimental/AnimNext/Source/AnimNextEditor/Private/Common/FindInAnimNextRigVMAsset.cpp b/Engine/Plugins/Experimental/AnimNext/Source/AnimNextEditor/Private/Common/FindInAnimNextRigVMAsset.cpp
new file mode 100644
--- /dev/null
+++ b/Engine/Plugins/Experimental/AnimNext/Source/AnimNextEditor/Private/Common/FindInAnimNextRigVMAsset.cpp
@@ -0,0 +1,98 @@
+// Copyright Epic Games, Inc. All Rights Reserved.
+
+#include "FindInAnimNextRigVMAsset.h"
+#include "IWorkspaceEditor.h"
+#include "AnimNextRigVMAsset.h"
+#include "AnimNextRigVMAssetEditorData.h"
+#include "RigVMModel/RigVMGraph.h"
+#include "EdGraph/RigVMEdGraph.h"
+#include "UncookedOnlyUtils.h"
+#include "GraphEditor.h"
+
+#define LOCTEXT_NAMESPACE "FindInAnimNextRigVMAsset"
+
+namespace UE::AnimNext::Editor
+{
+
+//////////////////////////////////////////////////////////////////////////
+// FFindInAnimNextRigVMAssetResult
+
+FFindInAnimNextRigVMAssetResult::FFindInAnimNextRigVMAssetResult(const FFindInGraphResult::FCreateParams& InCreateParams)
+	: FFindInGraphResult(InCreateParams)
+{
+}
+
+void FFindInAnimNextRigVMAssetResult::JumpToNode(TWeakPtr<FAssetEditorToolkit> AssetEditorToolkit, const UEdGraphNode* InNode) const
+{
+	if (!InNode)
+	{
+		return;
+	}
+
+	if (TSharedPtr<UE::Workspace::IWorkspaceEditor> WorkspaceEditor = StaticCastSharedPtr<UE::Workspace::IWorkspaceEditor>(AssetEditorToolkit.Pin()))
+	{		
+		// See if there already is an open graph editor, if not open it. Don't just always open as that can change focus contexts from parent modules
+		if (TSharedPtr<SGraphEditor> GraphEditor = SGraphEditor::FindGraphEditorForGraph(InNode->GetGraph()))
+		{
+			GraphEditor->JumpToNode(InNode, false);
+		}
+		else
+		{
+			WorkspaceEditor->OpenObjects({ const_cast<UEdGraph*>(InNode->GetGraph()) });
+			if (TSharedPtr<SGraphEditor> NewlyOpenedGraphEditor = SGraphEditor::FindGraphEditorForGraph(InNode->GetGraph()))
+			{
+				NewlyOpenedGraphEditor->JumpToNode(InNode, false);
+			}
+		}
+	}
+}
+
+
+//////////////////////////////////////////////////////////////////////////
+// SFindInAnimNextRigVMAsset
+
+TSharedPtr<FFindInGraphResult> SFindInAnimNextRigVMAsset::MakeSearchResult(const FFindInGraphResult::FCreateParams& InParams)
+{
+	return MakeShared<FFindInAnimNextRigVMAssetResult>(InParams);
+}
+
+void SFindInAnimNextRigVMAsset::MatchTokens(const TArray<FString>& Tokens)
+{
+	RootSearchResult.Reset();
+
+	if (TSharedPtr<UE::Workspace::IWorkspaceEditor> WorkspaceEditor = StaticCastSharedPtr<UE::Workspace::IWorkspaceEditor>(AssetEditorToolkitPtr.Pin()))
+	{
+		TObjectPtr<UObject> FocusedDocument = WorkspaceEditor->GetFocussedDocument();
+		if (UEdGraph* Graph = Cast<UEdGraph>(FocusedDocument))
+		{
+			// If we just selected a plain graph, rely on default graph search
+			MatchTokensInGraph(Graph, Tokens);
+		}
+		else if (UAnimNextRigVMAsset* AnimNextRigVMAsset = Cast<UAnimNextRigVMAsset>(FocusedDocument))
+		{
+			// Note: This is where we search subgraphs in say a module
+			if (UAnimNextRigVMAssetEditorData* AnimNextRigVMAssetEditorData = UE::AnimNext::UncookedOnly::FUtils::GetEditorData<UAnimNextRigVMAssetEditorData>(AnimNextRigVMAsset))
+			{
+				TArray<UEdGraph*> ContainedGraphs = AnimNextRigVMAssetEditorData->GetAllEdGraphs();
+				for (const UEdGraph* ContainedGraph : ContainedGraphs)
+				{
+					MatchTokensInGraph(ContainedGraph, Tokens);
+				}
+			}
+		}
+	}
+}
+
+bool SFindInAnimNextRigVMAsset::MatchTokensInNode(const UEdGraphNode* Node, const TArray<FString>& Tokens)
+{
+	// For now no type-specific node matching
+	return false;
+}
+
+
+/////////////////////////////////////////////////////
+
+} // namespace UE::AnimNext::Editor
+
+#undef LOCTEXT_NAMESPACE
+
diff --git a/Engine/Plugins/Experimental/AnimNext/Source/AnimNextEditor/Public/FindInAnimNextRigVMAsset.h b/Engine/Plugins/Experimental/AnimNext/Source/AnimNextEditor/Public/FindInAnimNextRigVMAsset.h
new file mode 100644
--- /dev/null
+++ b/Engine/Plugins/Experimental/AnimNext/Source/AnimNextEditor/Public/FindInAnimNextRigVMAsset.h
@@ -0,0 +1,40 @@
+// Copyright Epic Games, Inc. All Rights Reserved.
+
+#pragma once
+
+#include "CoreMinimal.h"
+#include "FindInGraph.h"
+
+class FAssetEditorToolkit;
+class SSearchBox;
+
+namespace UE::AnimNext::Editor
+{
+
+/** Item that matched the search results */
+class FFindInAnimNextRigVMAssetResult : public FFindInGraphResult
+{
+public:
+
+	FFindInAnimNextRigVMAssetResult(const FFindInGraphResult::FCreateParams& InCreateParams);
+
+	//~ Begin FFindInGraphResult Interface
+	virtual void JumpToNode(TWeakPtr<FAssetEditorToolkit> AssetEditorToolkit, const UEdGraphNode* InNode) const override;
+	//~ End FFindInGraphResult Interface
+};
+
+
+/** Widget for searching for items that are part of a AnimNextRigVMAsset Graph */
+class SFindInAnimNextRigVMAsset : public SFindInGraph
+{
+
+protected:
+
+	//~ Begin SFindInGraph Interface
+	virtual TSharedPtr<FFindInGraphResult> MakeSearchResult(const FFindInGraphResult::FCreateParams& InParams) override;
+	virtual void MatchTokens(const TArray<FString>& Tokens) override;
+	virtual bool MatchTokensInNode(const UEdGraphNode* Node, const TArray<FString>& Tokens) override;
+	//~ End SFindInGraph Interface
+};
+
+} // namespace UE::AnimNext::Editor
\ No newline at end of file
diff --git a/Engine/Plugins/Experimental/AnimNext/Source/AnimNextEditor/Public/IAnimNextEditorModule.h b/Engine/Plugins/Experimental/AnimNext/Source/AnimNextEditor/Public/IAnimNextEditorModule.h
--- a/Engine/Plugins/Experimental/AnimNext/Source/AnimNextEditor/Public/IAnimNextEditorModule.h
+++ b/Engine/Plugins/Experimental/AnimNext/Source/AnimNextEditor/Public/IAnimNextEditorModule.h
@@ -21,6 +21,7 @@ namespace UE::AnimNext::Editor
 
 const FLazyName CompilerResultsTabName("CompilerResultsTab");
 const FLazyName LogListingName("AnimNextCompilerResults");
+const FLazyName FindTabName("FindTab");
 
 struct FVariablePickerArgs;
 struct FActionMenuContextData;
diff --git a/Engine/Plugins/Experimental/AnimNext/Source/AnimNextUncookedOnly/Internal/AnimNextRigVMAssetEditorData.h b/Engine/Plugins/Experimental/AnimNext/Source/AnimNextUncookedOnly/Internal/AnimNextRigVMAssetEditorData.h
--- a/Engine/Plugins/Experimental/AnimNext/Source/AnimNextUncookedOnly/Internal/AnimNextRigVMAssetEditorData.h
+++ b/Engine/Plugins/Experimental/AnimNext/Source/AnimNextUncookedOnly/Internal/AnimNextRigVMAssetEditorData.h
@@ -32,6 +32,8 @@ namespace UE::AnimNext::Editor
 	class SParameterPicker;
 	class SRigVMAssetViewRow;
 	class FVariableCustomization;
+	class FFindInAnimNextRigVMAssetResult;
+	class SFindInAnimNextRigVMAsset;
 	class FAnimNextEditorModule;
 	class FWorkspaceEditor;
 	class FAnimNextAssetItemDetails;
@@ -169,6 +171,8 @@ protected:
 	friend class UE::AnimNext::Tests::FEditor_Variables;
 	friend class UE::AnimNext::Tests::FEditor_AnimGraph_Graphs;
 	friend class UE::AnimNext::Tests::FEditor_AnimGraph_Variables;
+	friend class UE::AnimNext::Editor::FFindInAnimNextRigVMAssetResult;
+	friend class UE::AnimNext::Editor::SFindInAnimNextRigVMAsset;
 	friend class UE::AnimNext::Tests::FVariables;
 	friend class UE::AnimNext::Tests::FVariables_UOLBindings;
 	friend class UE::AnimNext::Tests::FDataInterfaceCompile;
