Fix uncessary pragma once in .cpp AnimNextProgrammaticVariable

#tests nonunity

diff --git a/Engine/Plugins/Experimental/AnimNext/Source/AnimNextUncookedOnly/Private/Variables/AnimNextProgrammaticVariable.cpp b/Engine/Plugins/Experimental/AnimNext/Source/AnimNextUncookedOnly/Private/Variables/AnimNextProgrammaticVariable.cpp
--- a/Engine/Plugins/Experimental/AnimNext/Source/AnimNextUncookedOnly/Private/Variables/AnimNextProgrammaticVariable.cpp
+++ b/Engine/Plugins/Experimental/AnimNext/Source/AnimNextUncookedOnly/Private/Variables/AnimNextProgrammaticVariable.cpp
@@ -1,7 +1,5 @@
 // Copyright Epic Games, Inc. All Rights Reserved.
 
-#pragma once
-
 #include "Variables/AnimNextProgrammaticVariable.h"
 
 #include "Logging/StructuredLog.h"
