Make search also consider graphs referenced in nodes to match outliner / since variables can be used in referenced graphs.

#tests Unit, editor

diff --git a/Engine/Plugins/Experimental/AnimNext/Source/AnimNextEditor/Private/Common/FindInAnimNextRigVMAsset.cpp b/Engine/Plugins/Experimental/AnimNext/Source/AnimNextEditor/Private/Common/FindInAnimNextRigVMAsset.cpp
--- a/Engine/Plugins/Experimental/AnimNext/Source/AnimNextEditor/Private/Common/FindInAnimNextRigVMAsset.cpp
+++ b/Engine/Plugins/Experimental/AnimNext/Source/AnimNextEditor/Private/Common/FindInAnimNextRigVMAsset.cpp
@@ -2,12 +2,13 @@
 
 #include "FindInAnimNextRigVMAsset.h"
 #include "IWorkspaceEditor.h"
+#include "AnimNextEdGraphNode.h"
 #include "AnimNextRigVMAsset.h"
 #include "AnimNextRigVMAssetEditorData.h"
-#include "RigVMModel/RigVMGraph.h"
 #include "EdGraph/RigVMEdGraph.h"
-#include "UncookedOnlyUtils.h"
 #include "GraphEditor.h"
+#include "RigVMModel/RigVMGraph.h"
+#include "UncookedOnlyUtils.h"
 
 #define LOCTEXT_NAMESPACE "FindInAnimNextRigVMAsset"
 
@@ -84,8 +85,40 @@ void SFindInAnimNextRigVMAsset::MatchTokens(const TArray<FString>& Tokens)
 }
 
 bool SFindInAnimNextRigVMAsset::MatchTokensInNode(const UEdGraphNode* Node, const TArray<FString>& Tokens)
-{
-	// For now no type-specific node matching
+{	
+	// Search all animnext node pins to see if any graphs are referenced. If so add them to the search list
+	if (const UAnimNextEdGraphNode* AnimNextEdGraphNode = Cast<UAnimNextEdGraphNode>(Node))
+	{
+		for (const UEdGraphPin* Pin : AnimNextEdGraphNode->Pins)
+		{
+			if (Pin->Direction == EEdGraphPinDirection::EGPD_Input)
+			{
+				if (URigVMPin* RigVMPin = AnimNextEdGraphNode->FindModelPinFromGraphPin(Pin))
+				{
+					if (UClass* PinClass = Cast<UClass>(RigVMPin->GetCPPTypeObject()))
+					{
+						if (PinClass->IsChildOf<UAnimNextRigVMAsset>() || RigVMPin->GetMetaData("GetAllowedClasses") == "/Script/AnimNextAnimGraph.AnimNextAnimGraphSettings:GetAllowedAssetClasses")
+						{
+							// Search referenced asset if already loaded, don't change return value as this node still technically failed the search
+							const FSoftObjectPath ObjectPath(RigVMPin->GetDefaultValue());
+							if (UAnimNextRigVMAsset* AnimNextRigVMAsset = Cast<UAnimNextRigVMAsset>(ObjectPath.ResolveObject()))
+							{
+								if (UAnimNextRigVMAssetEditorData* AnimNextRigVMAssetEditorData = UE::AnimNext::UncookedOnly::FUtils::GetEditorData<UAnimNextRigVMAssetEditorData>(AnimNextRigVMAsset))
+								{
+									TArray<UEdGraph*> ContainedGraphs = AnimNextRigVMAssetEditorData->GetAllEdGraphs();
+									for (const UEdGraph* ContainedGraph : ContainedGraphs)
+									{
+										MatchTokensInGraph(ContainedGraph, Tokens);
+									}
+								}
+							}
+						}
+					}
+				}
+			}
+		}
+	}
+
 	return false;
 }
 
diff --git a/Engine/Plugins/Runtime/RigVM/Source/RigVMDeveloper/Private/RigVMModel/RigVMPin.cpp b/Engine/Plugins/Runtime/RigVM/Source/RigVMDeveloper/Private/RigVMModel/RigVMPin.cpp
--- a/Engine/Plugins/Runtime/RigVM/Source/RigVMDeveloper/Private/RigVMModel/RigVMPin.cpp
+++ b/Engine/Plugins/Runtime/RigVM/Source/RigVMDeveloper/Private/RigVMModel/RigVMPin.cpp
@@ -1306,7 +1306,7 @@ FString URigVMPin::GetMetaData(FName InKey) const
 		}
 		else
 		{
-			if(const UScriptStruct* Struct = UnitNode->GetScriptStruct())
+			if(const UScriptStruct* Struct = GetParentScriptStruct(UnitNode))
 			{
 				if(const FProperty* Property = Struct->FindPropertyByName(GetFName()))
 				{
@@ -1693,6 +1693,23 @@ UScriptStruct* URigVMPin::GetScriptStruct() const
 	return Cast<UScriptStruct>(GetCPPTypeObject());
 }
 
+UScriptStruct* URigVMPin::GetParentScriptStruct(const URigVMUnitNode* FallbackNode) const
+{
+	if (const URigVMPin* ParentPin = GetParentPin())
+	{
+		if (ParentPin->GetScriptStruct())
+		{
+			return ParentPin->GetScriptStruct();
+		}
+		else
+		{
+			return GetParentScriptStruct(FallbackNode);
+		}
+	}
+
+	return FallbackNode ? FallbackNode->GetScriptStruct() : nullptr;
+}
+
 UEnum* URigVMPin::GetEnum() const
 {
 	return Cast<UEnum>(GetCPPTypeObject());
diff --git a/Engine/Plugins/Runtime/RigVM/Source/RigVMDeveloper/Public/RigVMModel/RigVMPin.h b/Engine/Plugins/Runtime/RigVM/Source/RigVMDeveloper/Public/RigVMModel/RigVMPin.h
--- a/Engine/Plugins/Runtime/RigVM/Source/RigVMDeveloper/Public/RigVMModel/RigVMPin.h
+++ b/Engine/Plugins/Runtime/RigVM/Source/RigVMDeveloper/Public/RigVMModel/RigVMPin.h
@@ -365,6 +365,10 @@ public:
 	UFUNCTION(BlueprintCallable, Category = RigVMPin)
 	UScriptStruct* GetScriptStruct() const;
 
+	// Returns the parent struct of the data type of the Pin
+	UFUNCTION(BlueprintCallable, Category = RigVMPin)
+	UScriptStruct* GetParentScriptStruct(const URigVMUnitNode* FallbackNode) const;
+
 	// Returns the enum of the data type of the Pin,
 	// or nullptr otherwise.
 	UFUNCTION(BlueprintCallable, Category = RigVMPin)
