AnimNext convert TObjectPtr to TSoftObjectPtr for outliner item data to prevent auto-loading of workspace contained assets

diff --git a/Engine/Plugins/Experimental/AnimNext/Source/AnimNextEditor/Internal/Common/AnimNextAssetItemDetails.cpp b/Engine/Plugins/Experimental/AnimNext/Source/AnimNextEditor/Internal/Common/AnimNextAssetItemDetails.cpp
--- a/Engine/Plugins/Experimental/AnimNext/Source/AnimNextEditor/Internal/Common/AnimNextAssetItemDetails.cpp
+++ b/Engine/Plugins/Experimental/AnimNext/Source/AnimNextEditor/Internal/Common/AnimNextAssetItemDetails.cpp
@@ -77,7 +77,7 @@ void FAnimNextAssetItemDetails::RegisterToolMenuExtensions()
 		}
 
 		const FAnimNextRigVMAssetOutlinerData& OutlinerData = Data.Get<FAnimNextRigVMAssetOutlinerData>();
-		UAnimNextRigVMAsset* Asset = OutlinerData.Asset;
+		UAnimNextRigVMAsset* Asset = OutlinerData.GetAsset();
 		if(Asset == nullptr)
 		{
 			return;
@@ -238,7 +238,7 @@ void FAnimNextAssetItemDetails::RegisterToolMenuExtensions()
 							if(SubEntryClass == UAnimNextVariableEntry::StaticClass())
 							{
 								TSharedRef<SAddVariablesDialog> AddVariablesDialog =
-									SNew(SAddVariablesDialog, TArray<UAnimNextRigVMAssetEditorData*>({ UncookedOnly::FUtils::GetEditorData<UAnimNextRigVMAssetEditorData>(OutlinerData.Asset.Get()) }));
+									SNew(SAddVariablesDialog, TArray<UAnimNextRigVMAssetEditorData*>({ UncookedOnly::FUtils::GetEditorData<UAnimNextRigVMAssetEditorData>(OutlinerData.GetAsset()) }));
 
 								TArray<FVariableToAdd> VariablesToAdd;
 								TArray<FDataInterfaceToAdd> DataInterfacesToAdd;
@@ -258,7 +258,7 @@ void FAnimNextAssetItemDetails::RegisterToolMenuExtensions()
 							else if(SubEntryClass == UAnimNextDataInterfaceEntry::StaticClass())
 							{
 								TSharedRef<SAddVariablesDialog> AddVariablesDialog =
-									SNew(SAddVariablesDialog, TArray<UAnimNextRigVMAssetEditorData*>({ UncookedOnly::FUtils::GetEditorData<UAnimNextRigVMAssetEditorData>(OutlinerData.Asset.Get()) }))
+									SNew(SAddVariablesDialog, TArray<UAnimNextRigVMAssetEditorData*>({ UncookedOnly::FUtils::GetEditorData<UAnimNextRigVMAssetEditorData>(OutlinerData.GetAsset()) }))
 									.ShouldAddInitialVariable(false);
 
 								TArray<FVariableToAdd> VariablesToAdd;
diff --git a/Engine/Plugins/Experimental/AnimNext/Source/AnimNextEditor/Private/Common/AnimNextCollapseNodeItemDetails.cpp b/Engine/Plugins/Experimental/AnimNext/Source/AnimNextEditor/Private/Common/AnimNextCollapseNodeItemDetails.cpp
--- a/Engine/Plugins/Experimental/AnimNext/Source/AnimNextEditor/Private/Common/AnimNextCollapseNodeItemDetails.cpp
+++ b/Engine/Plugins/Experimental/AnimNext/Source/AnimNextEditor/Private/Common/AnimNextCollapseNodeItemDetails.cpp
@@ -32,9 +32,9 @@ bool FAnimNextCollapseNodeItemDetails::HandleDoubleClick(const FToolMenuContext&
 			if (Data.IsValid() && Data.GetScriptStruct() == FAnimNextCollapseGraphOutlinerData::StaticStruct())
 			{
 				const FAnimNextCollapseGraphOutlinerData& CollapseGraphData = Data.Get<FAnimNextCollapseGraphOutlinerData>();
-				if (CollapseGraphData.EditorObject.IsValid())
+				if (CollapseGraphData.SoftEditorObject.IsValid())
 				{
-					UObject* EditorObject = CollapseGraphData.EditorObject.Get();
+					UObject* EditorObject = CollapseGraphData.SoftEditorObject.Get();
 					WorkspaceEditor->OpenObjects({ EditorObject });
 					return true;
 				}
@@ -51,9 +51,9 @@ bool FAnimNextCollapseNodeItemDetails::CanDelete(const FWorkspaceOutlinerItemExp
 	if (Data.IsValid() && Data.GetScriptStruct() == FAnimNextCollapseGraphOutlinerData::StaticStruct())
 	{
 		const FAnimNextCollapseGraphOutlinerData& CollapseGraphData = Export.GetData().Get<FAnimNextCollapseGraphOutlinerData>();
-		if (CollapseGraphData.EditorObject.IsValid())
+		if (CollapseGraphData.SoftEditorObject.IsValid())
 		{
-			URigVMEdGraph* EdGraph = CollapseGraphData.EditorObject.Get();
+			URigVMEdGraph* EdGraph = CollapseGraphData.SoftEditorObject.Get();
 			if (EdGraph->bAllowDeletion)
 			{
 				return true;
@@ -72,9 +72,9 @@ void FAnimNextCollapseNodeItemDetails::Delete(TConstArrayView<FWorkspaceOutliner
 		if (Data.IsValid() && Data.GetScriptStruct() == FAnimNextCollapseGraphOutlinerData::StaticStruct())
 		{
 			const FAnimNextCollapseGraphOutlinerData& CollapseGraphData = Export.GetData().Get<FAnimNextCollapseGraphOutlinerData>();
-			if (CollapseGraphData.EditorObject.IsValid())
+			if (CollapseGraphData.SoftEditorObject.IsValid())
 			{
-				URigVMEdGraph* EdGraph = CollapseGraphData.EditorObject.Get();
+				URigVMEdGraph* EdGraph = CollapseGraphData.SoftEditorObject.Get();
 				if (EdGraph->bAllowDeletion)
 				{
 					if (URigVMGraph* Model = EdGraph->GetModel())
@@ -104,9 +104,9 @@ bool FAnimNextCollapseNodeItemDetails::CanRename(const FWorkspaceOutlinerItemExp
 	if (Data.IsValid() && Data.GetScriptStruct() == FAnimNextCollapseGraphOutlinerData::StaticStruct())
 	{
 		const FAnimNextCollapseGraphOutlinerData& CollapseGraphData = Export.GetData().Get<FAnimNextCollapseGraphOutlinerData>();
-		if (CollapseGraphData.EditorObject.IsValid())
+		if (CollapseGraphData.SoftEditorObject.IsValid())
 		{
-			URigVMEdGraph* EdGraph = CollapseGraphData.EditorObject.Get();
+			URigVMEdGraph* EdGraph = CollapseGraphData.SoftEditorObject.Get();
 			if (EdGraph->bAllowRenaming)
 			{
 				return true;
@@ -122,9 +122,9 @@ void FAnimNextCollapseNodeItemDetails::Rename(const FWorkspaceOutlinerItemExport
 	if (Data.IsValid() && Data.GetScriptStruct() == FAnimNextCollapseGraphOutlinerData::StaticStruct())
 	{
 		const FAnimNextCollapseGraphOutlinerData& CollapseGraphData = Export.GetData().Get<FAnimNextCollapseGraphOutlinerData>();
-		if (CollapseGraphData.EditorObject.IsValid())
+		if (CollapseGraphData.SoftEditorObject.IsValid())
 		{
-			URigVMEdGraph* EdGraph = CollapseGraphData.EditorObject.Get();
+			URigVMEdGraph* EdGraph = CollapseGraphData.SoftEditorObject.Get();
 			if (EdGraph->bAllowRenaming)
 			{
 				if (const UEdGraphSchema* GraphSchema = EdGraph->GetSchema())
@@ -168,9 +168,9 @@ UPackage* FAnimNextCollapseNodeItemDetails::GetPackage(const FWorkspaceOutlinerI
 	if (Data.IsValid() && Data.GetScriptStruct() == FAnimNextCollapseGraphOutlinerData::StaticStruct())
 	{
 		const FAnimNextCollapseGraphOutlinerData& CollapseGraphData = Data.Get<FAnimNextCollapseGraphOutlinerData>();
-		if (CollapseGraphData.EditorObject.IsValid())
+		if (CollapseGraphData.SoftEditorObject.IsValid())
 		{
-			return CollapseGraphData.EditorObject->GetPackage();
+			return CollapseGraphData.SoftEditorObject->GetPackage();
 		}
 	}
 	return nullptr;
diff --git a/Engine/Plugins/Experimental/AnimNext/Source/AnimNextEditor/Private/Common/AnimNextFunctionItemDetails.cpp b/Engine/Plugins/Experimental/AnimNext/Source/AnimNextEditor/Private/Common/AnimNextFunctionItemDetails.cpp
--- a/Engine/Plugins/Experimental/AnimNext/Source/AnimNextEditor/Private/Common/AnimNextFunctionItemDetails.cpp
+++ b/Engine/Plugins/Experimental/AnimNext/Source/AnimNextEditor/Private/Common/AnimNextFunctionItemDetails.cpp
@@ -32,9 +32,9 @@ bool FAnimNextFunctionItemDetails::HandleDoubleClick(const FToolMenuContext& Too
 			if (Data.IsValid() && Data.GetScriptStruct() == FAnimNextGraphFunctionOutlinerData::StaticStruct())
 			{
 				const FAnimNextGraphFunctionOutlinerData& FunctionData = Data.Get<FAnimNextGraphFunctionOutlinerData>();
-				if (FunctionData.EditorObject.IsValid())
+				if (FunctionData.SoftEditorObject.IsValid())
 				{
-					WorkspaceEditor->OpenObjects({ FunctionData.EditorObject.Get() });
+					WorkspaceEditor->OpenObjects({ FunctionData.SoftEditorObject.Get() });
 					return true;
 				}
 			}
@@ -55,9 +55,9 @@ bool FAnimNextFunctionItemDetails::HandleSelected(const FToolMenuContext& ToolMe
 			if (Data.IsValid() && Data.GetScriptStruct() == FAnimNextGraphFunctionOutlinerData::StaticStruct())
 			{
 				const FAnimNextGraphFunctionOutlinerData& FunctionData = Data.Get<FAnimNextGraphFunctionOutlinerData>();
-				if (FunctionData.EditorObject.IsValid())
+				if (FunctionData.SoftEditorObject.IsValid())
 				{
-					WorkspaceEditor->SetDetailsObjects({ FunctionData.EditorObject.Get() });
+					WorkspaceEditor->SetDetailsObjects({ FunctionData.SoftEditorObject.Get() });
 					return true;
 				}
 			}
@@ -73,9 +73,9 @@ bool FAnimNextFunctionItemDetails::CanDelete(const FWorkspaceOutlinerItemExport&
 	if (Data.IsValid() && Data.GetScriptStruct() == FAnimNextGraphFunctionOutlinerData::StaticStruct())
 	{
 		const FAnimNextGraphFunctionOutlinerData& FunctionData = Export.GetData().Get<FAnimNextGraphFunctionOutlinerData>();
-		if (FunctionData.EditorObject.IsValid())
+		if (FunctionData.SoftEditorObject.IsValid())
 		{
-			URigVMEdGraph* EdGraph = FunctionData.EditorObject.Get();
+			URigVMEdGraph* EdGraph = FunctionData.SoftEditorObject.Get();
 			if (EdGraph->bAllowDeletion)
 			{
 				return true;
@@ -94,11 +94,11 @@ void FAnimNextFunctionItemDetails::Delete(TConstArrayView<FWorkspaceOutlinerItem
 		if (Data.IsValid() && Data.GetScriptStruct() == FAnimNextGraphFunctionOutlinerData::StaticStruct())
 		{
 			const FAnimNextGraphFunctionOutlinerData& FunctionData = Export.GetData().Get<FAnimNextGraphFunctionOutlinerData>();
-			if (FunctionData.EditorObject.IsValid())
+			if (FunctionData.SoftEditorObject.IsValid())
 			{
-				if(FunctionData.EdGraphNode.IsValid())
+				if(FunctionData.SoftEdGraphNode.IsValid())
 				{
-					const URigVMEdGraphNode* EdGraphNode = FunctionData.EdGraphNode.Get();
+					const URigVMEdGraphNode* EdGraphNode = FunctionData.SoftEdGraphNode.Get();
 					if (EdGraphNode->CanUserDeleteNode())
 					{
 						if (URigVMGraph* Model = EdGraphNode->GetModel())
@@ -113,7 +113,7 @@ void FAnimNextFunctionItemDetails::Delete(TConstArrayView<FWorkspaceOutlinerItem
 				}
 				else
 				{
-					URigVMEdGraph* EdGraph = FunctionData.EditorObject.Get();
+					URigVMEdGraph* EdGraph = FunctionData.SoftEditorObject.Get();
 					if (EdGraph->bAllowDeletion)
 					{
 						if(IRigVMClientHost* ClientHost = EdGraph->GetImplementingOuter<IRigVMClientHost>())
@@ -138,9 +138,9 @@ bool FAnimNextFunctionItemDetails::CanRename(const FWorkspaceOutlinerItemExport&
 	if (Data.IsValid() && Data.GetScriptStruct() == FAnimNextGraphFunctionOutlinerData::StaticStruct())
 	{
 		const FAnimNextGraphFunctionOutlinerData& FunctionData = Export.GetData().Get<FAnimNextGraphFunctionOutlinerData>();
-		if (FunctionData.EditorObject.IsValid())
+		if (FunctionData.SoftEditorObject.IsValid())
 		{
-			URigVMEdGraph* EdGraph = FunctionData.EditorObject.Get();
+			URigVMEdGraph* EdGraph = FunctionData.SoftEditorObject.Get();
 			if (EdGraph->bAllowRenaming)
 			{
 				return true;
@@ -157,9 +157,9 @@ void FAnimNextFunctionItemDetails::Rename(const FWorkspaceOutlinerItemExport& Ex
 	if (Data.IsValid() && Data.GetScriptStruct() == FAnimNextGraphFunctionOutlinerData::StaticStruct())
 	{
 		const FAnimNextGraphFunctionOutlinerData& FunctionData = Export.GetData().Get<FAnimNextGraphFunctionOutlinerData>();
-		if (FunctionData.EditorObject.IsValid())
+		if (FunctionData.SoftEditorObject.IsValid())
 		{
-			URigVMEdGraph* EdGraph = FunctionData.EditorObject.Get();
+			URigVMEdGraph* EdGraph = FunctionData.SoftEditorObject.Get();
 			if (EdGraph->bAllowRenaming)
 			{
 				if (const UEdGraphSchema* GraphSchema = EdGraph->GetSchema())
@@ -203,9 +203,9 @@ UPackage* FAnimNextFunctionItemDetails::GetPackage(const FWorkspaceOutlinerItemE
 	if (Data.IsValid() && Data.GetScriptStruct() == FAnimNextGraphFunctionOutlinerData::StaticStruct())
 	{
 		const FAnimNextGraphFunctionOutlinerData& FunctionData = Data.Get<FAnimNextGraphFunctionOutlinerData>();
-		if (FunctionData.EditorObject.IsValid())
+		if (FunctionData.SoftEditorObject.IsValid())
 		{
-			return FunctionData.EditorObject->GetPackage();
+			return FunctionData.SoftEditorObject->GetPackage();
 		}
 	}
 	return nullptr;
diff --git a/Engine/Plugins/Experimental/AnimNext/Source/AnimNextEditor/Private/Common/AnimNextGraphItemDetails.cpp b/Engine/Plugins/Experimental/AnimNext/Source/AnimNextEditor/Private/Common/AnimNextGraphItemDetails.cpp
--- a/Engine/Plugins/Experimental/AnimNext/Source/AnimNextEditor/Private/Common/AnimNextGraphItemDetails.cpp
+++ b/Engine/Plugins/Experimental/AnimNext/Source/AnimNextEditor/Private/Common/AnimNextGraphItemDetails.cpp
@@ -32,9 +32,9 @@ bool FAnimNextGraphItemDetails::HandleDoubleClick(const FToolMenuContext& ToolMe
 			if (Data.IsValid() && Data.GetScriptStruct() == FAnimNextGraphOutlinerData::StaticStruct())
 			{
 				const FAnimNextGraphOutlinerData& GraphData = Data.Get<FAnimNextGraphOutlinerData>();							
-				if (GraphData.GraphInterface)
+				if (GraphData.GetGraphInterface())
 				{
-					if (URigVMGraph* RigVMGraph = GraphData.GraphInterface->GetRigVMGraph())
+					if (URigVMGraph* RigVMGraph = GraphData.GetGraphInterface()->GetRigVMGraph())
 					{
 						if(const IRigVMClientHost* RigVMClientHost = RigVMGraph->GetImplementingOuter<IRigVMClientHost>())
 						{
@@ -73,12 +73,12 @@ void FAnimNextGraphItemDetails::Delete(TConstArrayView<FWorkspaceOutlinerItemExp
 		{
 			const FAnimNextGraphOutlinerData& GraphData = Export.GetData().Get<FAnimNextGraphOutlinerData>();
 
-			UAnimNextRigVMAssetEditorData* EditorData = GraphData.Entry->GetTypedOuter<UAnimNextRigVMAssetEditorData>();
+			UAnimNextRigVMAssetEditorData* EditorData = GraphData.GetEntry()->GetTypedOuter<UAnimNextRigVMAssetEditorData>();
 			if(EditorData == nullptr)
 			{
 				continue;
 			}
-			UAnimNextRigVMAssetEntry* Entry = GraphData.Entry;
+			UAnimNextRigVMAssetEntry* Entry = GraphData.GetEntry();
 			if(Entry == nullptr)
 			{
 				continue;
@@ -114,13 +114,13 @@ void FAnimNextGraphItemDetails::Rename(const FWorkspaceOutlinerItemExport& Expor
 	if (Data.IsValid() && Data.GetScriptStruct() == FAnimNextGraphOutlinerData::StaticStruct())
 	{
 		const FAnimNextGraphOutlinerData& GraphData = Export.GetData().Get<FAnimNextGraphOutlinerData>();
-		UAnimNextRigVMAssetEditorData* EditorData = GraphData.Entry->GetTypedOuter<UAnimNextRigVMAssetEditorData>();
+		UAnimNextRigVMAssetEditorData* EditorData = GraphData.GetEntry()->GetTypedOuter<UAnimNextRigVMAssetEditorData>();
 		FName NewName = FName(*InName.ToString());
 		UAnimNextRigVMAssetEntry* ExistingEntry = EditorData->FindEntry(NewName);
-		if (ExistingEntry == nullptr && GraphData.Entry->GetEntryName() != NewName)
+		if (ExistingEntry == nullptr && GraphData.GetEntry()->GetEntryName() != NewName)
 		{
 			FScopedTransaction Transaction(LOCTEXT("SetName", "Set Name"));
-			GraphData.Entry->SetEntryName(NewName);
+			GraphData.GetEntry()->SetEntryName(NewName);
 		}
 	}
 }
@@ -131,9 +131,9 @@ bool FAnimNextGraphItemDetails::ValidateName(const FWorkspaceOutlinerItemExport&
 	if (Data.IsValid() && Data.GetScriptStruct() == FAnimNextGraphOutlinerData::StaticStruct())
 	{
 		const FAnimNextGraphOutlinerData& GraphData = Export.GetData().Get<FAnimNextGraphOutlinerData>();
-		UAnimNextRigVMAssetEditorData* EditorData = GraphData.Entry->GetTypedOuter<UAnimNextRigVMAssetEditorData>();
+		UAnimNextRigVMAssetEditorData* EditorData = GraphData.GetEntry()->GetTypedOuter<UAnimNextRigVMAssetEditorData>();
 		UAnimNextRigVMAssetEntry* ExistingEntry = EditorData->FindEntry(FName(*InName.ToString()));
-		if(ExistingEntry && GraphData.Entry != ExistingEntry)
+		if(ExistingEntry && GraphData.GetEntry() != ExistingEntry)
 		{
 			OutErrorMessage = LOCTEXT("NameAlreadyExistsError", "Name already exists in this module");
 			return false;
@@ -152,9 +152,9 @@ UPackage* FAnimNextGraphItemDetails::GetPackage(const FWorkspaceOutlinerItemExpo
 	if (Data.IsValid() && Data.GetScriptStruct() == FAnimNextGraphOutlinerData::StaticStruct())
 	{
 		const FAnimNextGraphOutlinerData& GraphData = Data.Get<FAnimNextGraphOutlinerData>();
-		if (GraphData.GraphInterface)
+		if (GraphData.SoftEntryPtr.IsValid())
 		{
-			return GraphData.GraphInterface.GetObject()->GetExternalPackage();
+			return GraphData.GetEntry()->GetExternalPackage();
 		}
 	}
 	return nullptr;
diff --git a/Engine/Plugins/Experimental/AnimNext/Source/AnimNextEditor/Private/Variables/SVariablesView.cpp b/Engine/Plugins/Experimental/AnimNext/Source/AnimNextEditor/Private/Variables/SVariablesView.cpp
--- a/Engine/Plugins/Experimental/AnimNext/Source/AnimNextEditor/Private/Variables/SVariablesView.cpp
+++ b/Engine/Plugins/Experimental/AnimNext/Source/AnimNextEditor/Private/Variables/SVariablesView.cpp
@@ -268,9 +268,9 @@ void SVariablesView::HandleWorkspaceOutlinerSelectionChanged(TConstArrayView<FWo
 			// variables. Eventually we would want to recurse up to the root of the export hierarchy, which will entail changing the signature of
 			// FOnOutlinerSelectionChanged to include parent items somehow, or adding parent query functionality to the workspace API
 			const FAnimNextAssetEntryOutlinerData& EntryData = Export.GetData().Get<FAnimNextAssetEntryOutlinerData>();
-			if(EntryData.Entry)
+			if(EntryData.SoftEntryPtr.IsValid())
 			{
-				if(UAnimNextRigVMAsset* Asset = EntryData.Entry->GetTypedOuter<UAnimNextRigVMAsset>())
+				if(UAnimNextRigVMAsset* Asset = EntryData.GetEntry()->GetTypedOuter<UAnimNextRigVMAsset>())
 				{
 					Assets.Add(TSoftObjectPtr<UAnimNextRigVMAsset>(Asset));
 				}
@@ -279,9 +279,9 @@ void SVariablesView::HandleWorkspaceOutlinerSelectionChanged(TConstArrayView<FWo
 		else if(Export.HasData() && Export.GetData().GetScriptStruct()->IsChildOf(FAnimNextGraphFunctionOutlinerData::StaticStruct()))
 		{
 			const FAnimNextGraphFunctionOutlinerData& EntryData = Export.GetData().Get<FAnimNextGraphFunctionOutlinerData>();
-			if(EntryData.EditorObject.IsValid())
+			if(EntryData.SoftEditorObject.IsValid())
 			{
-				if(URigVMEdGraph* EdGraph = EntryData.EditorObject.Get())
+				if(URigVMEdGraph* EdGraph = EntryData.SoftEditorObject.Get())
 				{
 					if(UAnimNextRigVMAsset* Asset = EdGraph->GetTypedOuter<UAnimNextRigVMAsset>())
 					{
diff --git a/Engine/Plugins/Experimental/AnimNext/Source/AnimNextUncookedOnly/Internal/AnimNextAssetWorkspaceAssetUserData.h b/Engine/Plugins/Experimental/AnimNext/Source/AnimNextUncookedOnly/Internal/AnimNextAssetWorkspaceAssetUserData.h
--- a/Engine/Plugins/Experimental/AnimNext/Source/AnimNextUncookedOnly/Internal/AnimNextAssetWorkspaceAssetUserData.h
+++ b/Engine/Plugins/Experimental/AnimNext/Source/AnimNextUncookedOnly/Internal/AnimNextAssetWorkspaceAssetUserData.h
@@ -22,9 +22,14 @@ struct FAnimNextRigVMAssetOutlinerData : public FWorkspaceOutlinerItemData
 	GENERATED_BODY()
 
 	FAnimNextRigVMAssetOutlinerData() = default;
+	
+	UAnimNextRigVMAsset* GetAsset() const
+	{
+		return SoftAssetPtr.LoadSynchronous();
+	}
 
 	UPROPERTY(VisibleAnywhere, Category=AnimNext)
-	TObjectPtr<UAnimNextRigVMAsset> Asset;
+	TSoftObjectPtr<UAnimNextRigVMAsset> SoftAssetPtr;
 };
 
 USTRUCT()
@@ -36,7 +41,7 @@ struct FAnimNextModuleOutlinerData : public FAnimNextRigVMAssetOutlinerData
 
 	UAnimNextModule* GetModule() const
 	{
-		return Cast<UAnimNextModule>(Asset);
+		return Cast<UAnimNextModule>(GetAsset());
 	}
 };
 
@@ -47,9 +52,9 @@ struct FAnimNextDataInterfaceOutlinerData : public FAnimNextRigVMAssetOutlinerDa
 
 	FAnimNextDataInterfaceOutlinerData() = default;
 
-	UAnimNextDataInterface* GetDataInteface() const
+	UAnimNextDataInterface* GetDataInterface() const
 	{
-		return Cast<UAnimNextDataInterface>(Asset);
+		return Cast<UAnimNextDataInterface>(GetAsset());
 	}
 };
 
@@ -60,9 +65,14 @@ struct FAnimNextAssetEntryOutlinerData : public FWorkspaceOutlinerItemData
 	GENERATED_BODY()
 	
 	FAnimNextAssetEntryOutlinerData() = default;
+	
+	UAnimNextRigVMAssetEntry* GetEntry() const
+	{
+		return SoftEntryPtr.LoadSynchronous();
+	}
 
 	UPROPERTY(VisibleAnywhere, Category=AnimNext)
-	TObjectPtr<UAnimNextRigVMAssetEntry> Entry;
+	TSoftObjectPtr<UAnimNextRigVMAssetEntry> SoftEntryPtr;
 };
 
 USTRUCT()
@@ -84,7 +94,7 @@ struct FAnimNextCollapseGraphsOutlinerDataBase : public FWorkspaceOutlinerItemDa
 	FAnimNextCollapseGraphsOutlinerDataBase() = default;
 
 	UPROPERTY(VisibleAnywhere, Category=AnimNext)
-	TWeakObjectPtr<URigVMEdGraph> EditorObject;
+	TSoftObjectPtr<URigVMEdGraph> SoftEditorObject;
 };
 
 USTRUCT()
@@ -103,7 +113,7 @@ struct FAnimNextGraphFunctionOutlinerData : public FAnimNextCollapseGraphsOutlin
 	FAnimNextGraphFunctionOutlinerData() = default;
 
 	UPROPERTY(VisibleAnywhere, Category=AnimNext)
-	TWeakObjectPtr<URigVMEdGraphNode> EdGraphNode;
+	TSoftObjectPtr<URigVMEdGraphNode> SoftEdGraphNode;
 };
 
 USTRUCT()
@@ -112,9 +122,19 @@ struct FAnimNextGraphOutlinerData : public FAnimNextAssetEntryOutlinerData
 	GENERATED_BODY()
 	
 	FAnimNextGraphOutlinerData() = default;
-
-	UPROPERTY(VisibleAnywhere, Category=AnimNext)
-	TScriptInterface<IAnimNextRigVMGraphInterface> GraphInterface;
+	
+	IAnimNextRigVMGraphInterface* GetGraphInterface() const
+	{
+		if (UAnimNextRigVMAssetEntry* Entry = GetEntry())
+		{
+			if (IAnimNextRigVMGraphInterface* GraphInterface = CastChecked<IAnimNextRigVMGraphInterface>(Entry))
+			{
+				return GraphInterface;
+			}
+		}
+
+		return nullptr;
+	}
 };
 
 UCLASS(MinimalAPI)
diff --git a/Engine/Plugins/Experimental/AnimNext/Source/AnimNextUncookedOnly/Private/AnimNextAssetWorkspaceAssetUserData.cpp b/Engine/Plugins/Experimental/AnimNext/Source/AnimNextUncookedOnly/Private/AnimNextAssetWorkspaceAssetUserData.cpp
--- a/Engine/Plugins/Experimental/AnimNext/Source/AnimNextUncookedOnly/Private/AnimNextAssetWorkspaceAssetUserData.cpp
+++ b/Engine/Plugins/Experimental/AnimNext/Source/AnimNextUncookedOnly/Private/AnimNextAssetWorkspaceAssetUserData.cpp
@@ -32,7 +32,7 @@ void UAnimNextAssetWorkspaceAssetUserData::GetAssetRegistryTags(FAssetRegistryTa
 				RootAssetExport.GetData().InitializeAsScriptStruct(FAnimNextRigVMAssetOutlinerData::StaticStruct());
 			}
 			FAnimNextRigVMAssetOutlinerData& Data = RootAssetExport.GetData().GetMutable<FAnimNextRigVMAssetOutlinerData>();
-			Data.Asset = Asset;
+			Data.SoftAssetPtr = Asset;
 		}
 	
 		UE::AnimNext::UncookedOnly::FUtils::GetAssetOutlinerItems(GraphEditorData, Exports, Context);
diff --git a/Engine/Plugins/Experimental/AnimNext/Source/AnimNextUncookedOnly/Private/UncookedOnlyUtils.cpp b/Engine/Plugins/Experimental/AnimNext/Source/AnimNextUncookedOnly/Private/UncookedOnlyUtils.cpp
--- a/Engine/Plugins/Experimental/AnimNext/Source/AnimNextUncookedOnly/Private/UncookedOnlyUtils.cpp
+++ b/Engine/Plugins/Experimental/AnimNext/Source/AnimNextUncookedOnly/Private/UncookedOnlyUtils.cpp
@@ -933,8 +933,7 @@ void FUtils::GetAssetOutlinerItems(const UAnimNextRigVMAssetEditorData* EditorDa
 
 				Export.GetData().InitializeAsScriptStruct(FAnimNextGraphOutlinerData::StaticStruct());
 				FAnimNextGraphOutlinerData& GraphData = Export.GetData().GetMutable<FAnimNextGraphOutlinerData>();
-				GraphData.Entry = Entry;
-				GraphData.GraphInterface = GraphInterface->_getUObject();
+				GraphData.SoftEntryPtr = Entry;
 
 				if (URigVMEdGraph* RigVMEdGraph = GraphInterface->GetEdGraph())
 				{
@@ -1029,7 +1028,7 @@ void FUtils::CreateSubGraphsOutlinerItemsRecursive(const UAnimNextRigVMAssetEdit
 				Export.GetData().InitializeAsScriptStruct(FAnimNextCollapseGraphOutlinerData::StaticStruct());
 			
 				FAnimNextCollapseGraphOutlinerData& FnGraphData = Export.GetData().GetMutable<FAnimNextCollapseGraphOutlinerData>();
-				FnGraphData.EditorObject = EditorObject;
+				FnGraphData.SoftEditorObject = EditorObject;
 
 				CreateSubGraphsOutlinerItemsRecursive(EditorData, OutExports, Export, EditorObject, Context);
 			}
@@ -1044,24 +1043,21 @@ void FUtils::CreateSubGraphsOutlinerItemsRecursive(const UAnimNextRigVMAssetEdit
 	{
 		if (URigVMFunctionReferenceNode* FunctionReferenceNode = Cast<URigVMFunctionReferenceNode>(EdNode->GetModelNode()))
 		{
-			if (URigVMLibraryNode* ReferencedNode = Cast<URigVMLibraryNode>(FunctionReferenceNode->GetReferencedFunctionHeader().LibraryPointer.GetNodeSoftPath().ResolveObject()))
+			// Only export referenced functions which are part of same outer
+			const FRigVMGraphFunctionIdentifier FunctionIdentifier = FunctionReferenceNode->GetFunctionIdentifier();
+			if (FunctionIdentifier.HostObject == FSoftObjectPath(EditorData))
 			{
-				if (URigVMGraph* ContainedGraph = ReferencedNode->GetContainedGraph())
+				const URigVMLibraryNode* FunctionNode = EditorData->RigVMClient.GetFunctionLibrary()->FindFunction(FunctionIdentifier.GetFunctionFName());
+				if (URigVMGraph* ContainedGraph = FunctionNode->GetContainedGraph())
 				{
-					if (EditorData->GetEditorObjectForRigVMGraph(ContainedGraph) == nullptr)
-					{
-						continue; // Do not show references to other assets functions in the outliner
-					}
-
-					FWorkspaceOutlinerItemExport& Export = OutExports.Exports.Add_GetRef(FWorkspaceOutlinerItemExport(ReferencedNode->GetFName(), ParentExport));
-
+					FWorkspaceOutlinerItemExport& Export = OutExports.Exports.Add_GetRef(FWorkspaceOutlinerItemExport(FunctionNode->GetFName(), ParentExport));
 					Export.GetData().InitializeAsScriptStruct(FAnimNextGraphFunctionOutlinerData::StaticStruct());
 					FAnimNextGraphFunctionOutlinerData& FnGraphData = Export.GetData().GetMutable<FAnimNextGraphFunctionOutlinerData>();
 
 					if (URigVMEdGraph* ContainedEdGraph = Cast<URigVMEdGraph>(EditorData->GetEditorObjectForRigVMGraph(ContainedGraph)))
 					{
-						FnGraphData.EditorObject = ContainedEdGraph;
-						FnGraphData.EdGraphNode = EdNode;
+						FnGraphData.SoftEditorObject = ContainedEdGraph;
+						FnGraphData.SoftEdGraphNode = EdNode;
 						
 						CreateSubGraphsOutlinerItemsRecursive(EditorData, OutExports, Export, ContainedEdGraph, Context);
 					}
@@ -1084,9 +1080,10 @@ void FUtils::CreateFunctionsOutlinerItemsRecursive(const UAnimNextRigVMAssetEdit
 {
 	for (const FRigVMGraphFunctionData& FunctionData : Functions)
 	{
-		if (URigVMLibraryNode* LibraryNode = Cast<URigVMLibraryNode>(FunctionData.Header.LibraryPointer.GetNodeSoftPath().ResolveObject()))
+		const URigVMLibraryNode* FunctionNode = EditorData->RigVMClient.GetFunctionLibrary()->FindFunction(FunctionData.Header.LibraryPointer.GetFunctionFName());
+		if (FunctionNode)
 		{
-			if (URigVMGraph* ContainedModelGraph = LibraryNode->GetContainedGraph())
+			if (URigVMGraph* ContainedModelGraph = FunctionNode->GetContainedGraph())
 			{
 				if (URigVMEdGraph* EditorObject = Cast<URigVMEdGraph>(EditorData->GetEditorObjectForRigVMGraph(ContainedModelGraph)))
 				{
@@ -1094,7 +1091,7 @@ void FUtils::CreateFunctionsOutlinerItemsRecursive(const UAnimNextRigVMAssetEdit
 
 					Export.GetData().InitializeAsScriptStruct(FAnimNextGraphFunctionOutlinerData::StaticStruct());
 					FAnimNextGraphFunctionOutlinerData& FnGraphData = Export.GetData().GetMutable<FAnimNextGraphFunctionOutlinerData>();
-					FnGraphData.EditorObject = EditorObject;
+					FnGraphData.SoftEditorObject = EditorObject;
 				}
 			}
 		}
diff --git a/Engine/Plugins/Experimental/AnimNextAnimGraph/Source/AnimNextAnimGraphUncookedOnly/Internal/AnimNextAnimGraphWorkspaceAssetUserData.h b/Engine/Plugins/Experimental/AnimNextAnimGraph/Source/AnimNextAnimGraphUncookedOnly/Internal/AnimNextAnimGraphWorkspaceAssetUserData.h
--- a/Engine/Plugins/Experimental/AnimNextAnimGraph/Source/AnimNextAnimGraphUncookedOnly/Internal/AnimNextAnimGraphWorkspaceAssetUserData.h
+++ b/Engine/Plugins/Experimental/AnimNextAnimGraph/Source/AnimNextAnimGraphUncookedOnly/Internal/AnimNextAnimGraphWorkspaceAssetUserData.h
@@ -17,7 +17,7 @@ struct FAnimNextAnimationGraphOutlinerData : public FAnimNextRigVMAssetOutlinerD
 
 	UAnimNextAnimationGraph* GetAnimationGraph() const
 	{
-		return Cast<UAnimNextAnimationGraph>(Asset);
+		return Cast<UAnimNextAnimationGraph>(GetAsset());
 	}
 };
 
diff --git a/Engine/Plugins/Experimental/AnimNextAnimGraph/Source/AnimNextAnimGraphUncookedOnly/Private/AnimNextAnimGraphWorkspaceAssetUserData.cpp b/Engine/Plugins/Experimental/AnimNextAnimGraph/Source/AnimNextAnimGraphUncookedOnly/Private/AnimNextAnimGraphWorkspaceAssetUserData.cpp
--- a/Engine/Plugins/Experimental/AnimNextAnimGraph/Source/AnimNextAnimGraphUncookedOnly/Private/AnimNextAnimGraphWorkspaceAssetUserData.cpp
+++ b/Engine/Plugins/Experimental/AnimNextAnimGraph/Source/AnimNextAnimGraphUncookedOnly/Private/AnimNextAnimGraphWorkspaceAssetUserData.cpp
@@ -18,7 +18,7 @@ void UAnimNextAnimGraphWorkspaceAssetUserData::GetAssetRegistryTags(FAssetRegist
 			FWorkspaceOutlinerItemExport& RootAssetExport = Exports.Exports.Add_GetRef(FWorkspaceOutlinerItemExport(Asset->GetFName(), Asset));
 			RootAssetExport.GetData().InitializeAsScriptStruct(FAnimNextAnimationGraphOutlinerData::StaticStruct());
 			FAnimNextRigVMAssetOutlinerData& Data = RootAssetExport.GetData().GetMutable<FAnimNextRigVMAssetOutlinerData>();
-			Data.Asset = Asset;
+			Data.SoftAssetPtr = Asset;
 		}
 
 		UE::AnimNext::UncookedOnly::FUtils::GetAssetOutlinerItems(EditorData, Exports, Context);
diff --git a/Engine/Plugins/Experimental/AnimNextStateTree/Source/AnimNextStateTreeUncookedOnly/Private/AnimNextStateTreeWorkspaceAssetUserData.cpp b/Engine/Plugins/Experimental/AnimNextStateTree/Source/AnimNextStateTreeUncookedOnly/Private/AnimNextStateTreeWorkspaceAssetUserData.cpp
--- a/Engine/Plugins/Experimental/AnimNextStateTree/Source/AnimNextStateTreeUncookedOnly/Private/AnimNextStateTreeWorkspaceAssetUserData.cpp
+++ b/Engine/Plugins/Experimental/AnimNextStateTree/Source/AnimNextStateTreeUncookedOnly/Private/AnimNextStateTreeWorkspaceAssetUserData.cpp
@@ -21,7 +21,7 @@ void UAnimNextStateTreeWorkspaceAssetUserData::GetAssetRegistryTags(FAssetRegist
 		UStateTree* StateTree = AnimStateTree->StateTree;
 		FWorkspaceOutlinerItemExport& RootAssetExport = OutlinerExports.Exports.Add_GetRef(FWorkspaceOutlinerItemExport(AnimStateTreeIdentifier, GetOuter()));			
 		RootAssetExport.GetData().InitializeAsScriptStruct(FAnimNextStateTreeOutlinerData::StaticStruct());
-		RootAssetExport.GetData().GetMutable<FAnimNextStateTreeOutlinerData>().Asset = AnimStateTree;
+		RootAssetExport.GetData().GetMutable<FAnimNextStateTreeOutlinerData>().SoftAssetPtr = AnimStateTree;
 
 		if (const UStateTreeEditorData* EditorData = Cast<UStateTreeEditorData>(StateTree->EditorData)) 
 		{
diff --git a/Engine/Plugins/Experimental/Workspace/Source/WorkspaceEditor/Private/Workspace.cpp b/Engine/Plugins/Experimental/Workspace/Source/WorkspaceEditor/Private/Workspace.cpp
--- a/Engine/Plugins/Experimental/Workspace/Source/WorkspaceEditor/Private/Workspace.cpp
+++ b/Engine/Plugins/Experimental/Workspace/Source/WorkspaceEditor/Private/Workspace.cpp
@@ -274,7 +274,7 @@ bool UWorkspace::HasValidEntries() const
 {
 	for(const UWorkspaceAssetEntry* AssetEntry : AssetEntries)
 	{
-		if (AssetEntry && AssetEntry->Asset.IsValid())
+		if (AssetEntry && !AssetEntry->Asset.IsNull())
 		{
 			return true;
 		}
diff --git a/Engine/Plugins/Experimental/Workspace/Source/WorkspaceEditor/Private/WorkspaceEditor.cpp b/Engine/Plugins/Experimental/Workspace/Source/WorkspaceEditor/Private/WorkspaceEditor.cpp
--- a/Engine/Plugins/Experimental/Workspace/Source/WorkspaceEditor/Private/WorkspaceEditor.cpp
+++ b/Engine/Plugins/Experimental/Workspace/Source/WorkspaceEditor/Private/WorkspaceEditor.cpp
@@ -750,14 +750,18 @@ bool FWorkspaceEditor::AreAssetEntriesModified() const
 	{
 		if (Entry != nullptr)
 		{
-			if(UObject* Asset = Entry->Asset.Get())
+			if (Entry->Asset.IsValid())
 			{
-				const UPackage* Package = Asset->GetOutermost();
-				if (Package->IsDirty() || Package->GetExternalPackages().ContainsByPredicate([](const UPackage* ExternalPackage)  { return ExternalPackage && ExternalPackage->IsDirty(); }))
+				if (UObject* Asset = Entry->Asset.Get())
 				{
-					return true;
+					const UPackage* Package = Asset->GetOutermost();
+					if (Package->IsDirty() || Package->GetExternalPackages().ContainsByPredicate([](const UPackage* ExternalPackage)  { return ExternalPackage && ExternalPackage->IsDirty(); }))
+					{
+						return true;
+					}
 				}
 			}
+			
 		}
 	}
 	
