[AnimRuntime] Translational retargeting settings can now be inherited from the compatible skeleton. This allows you to setup the per bone translational retarget settings in the retarget manager just once for your core animations, and have different characters that are compatible with that skeleton play those animations and automatically use the per bone retarget modes from the source/core skeleton. This means you do not have to setup the retarget modes for every character anymore. This is optional though. There is a new checkbox/boolean that allows you to enable this. It is disabled on default. You can find this setting/checkbox in the retarget manager in the skeleton asset editor, next to where you set up your compatible skeletons.

diff --git a/Engine/Plugins/Experimental/AnimNext/Source/AnimNext/Private/DecompressionTools.cpp b/Engine/Plugins/Experimental/AnimNext/Source/AnimNext/Private/DecompressionTools.cpp
--- a/Engine/Plugins/Experimental/AnimNext/Source/AnimNext/Private/DecompressionTools.cpp
+++ b/Engine/Plugins/Experimental/AnimNext/Source/AnimNext/Private/DecompressionTools.cpp
@@ -317,8 +317,12 @@ void FDecompressionTools::DecompressPose(FLODPose& OutAnimationPoseData,
 
 	const int32 NumTracks = CompressedData.CompressedTrackToSkeletonMapTable.Num();
 
+	const USkeleton* SourceSkeleton = DecompressionContext.GetSourceSkeleton();
 	const USkeleton* TargetSkeleton = OutAnimationPoseData.GetSkeletonAsset();
-	const FSkeletonRemapping& SkeletonRemapping = UE::Anim::FSkeletonRemappingRegistry::Get().GetRemapping(DecompressionContext.GetSourceSkeleton(), TargetSkeleton);
+	const FSkeletonRemapping& SkeletonRemapping = UE::Anim::FSkeletonRemappingRegistry::Get().GetRemapping(SourceSkeleton, TargetSkeleton);
+
+	const bool bUseSourceRetargetModes = TargetSkeleton->GetUseRetargetModesFromCompatibleSkeleton();
+	const bool bDisableRetargeting = OutAnimationPoseData.GetDisableRetargeting();
 
 	FGetBonePoseScratchArea& ScratchArea = FGetBonePoseScratchArea::Get();
 	BoneTrackArray& RotationScalePairs = ScratchArea.RotationScalePairs;
@@ -368,8 +372,17 @@ void FDecompressionTools::DecompressPose(FLODPose& OutAnimationPoseData,
 						AnimatedCompactRotations[LODBoneIndex] = true;
 					}
 
+					// Check what retarget mode to use for the translational retargeting for this specific bone.
+					const EBoneTranslationRetargetingMode::Type RetargetMode = FAnimationRuntime::GetBoneTranslationRetargetingMode(
+						bUseSourceRetargetModes,
+						SourceSkeletonBoneIndex,
+						TargetSkeletonBoneIndex,
+						SourceSkeleton,
+						TargetSkeleton,
+						bDisableRetargeting);
+
 					// Skip extracting translation component for EBoneTranslationRetargetingMode::Skeleton.
-					switch (TargetSkeleton->GetBoneTranslationRetargetingMode(TargetSkeletonBoneIndex, OutAnimationPoseData.GetDisableRetargeting()))
+					switch (RetargetMode)
 					{
 					case EBoneTranslationRetargetingMode::Animation:
 						TranslationPairs.Add(BoneTrackPair(LODBoneIndex, TrackIndex));
@@ -425,12 +438,21 @@ void FDecompressionTools::DecompressPose(FLODPose& OutAnimationPoseData,
 				// Root bone does not require fix-up for additive animations as there is no parent delta rotation to account for
 				if (!DecompressionContext.IsAdditiveAnimation())
 				{
-					const int32 TargetSkeletonBoneIndex = 0;
-
-					RootAtom.SetRotation(SkeletonRemapping.RetargetBoneRotationToTargetSkeleton(TargetSkeletonBoneIndex, RootAtom.GetRotation()));
-					if (TargetSkeleton->GetBoneTranslationRetargetingMode(TargetSkeletonBoneIndex, OutAnimationPoseData.GetDisableRetargeting()) != EBoneTranslationRetargetingMode::Skeleton)
+					constexpr int32 SkeletonBoneIndex = 0;
+
+					// Check what retarget mode to use for the translational retargeting for this specific bone.
+					const EBoneTranslationRetargetingMode::Type RetargetMode = FAnimationRuntime::GetBoneTranslationRetargetingMode(
+						bUseSourceRetargetModes,
+						SkeletonBoneIndex,
+						SkeletonBoneIndex,
+						SourceSkeleton,
+						TargetSkeleton,
+						bDisableRetargeting);
+
+					RootAtom.SetRotation(SkeletonRemapping.RetargetBoneRotationToTargetSkeleton(SkeletonBoneIndex, RootAtom.GetRotation()));
+					if (RetargetMode != EBoneTranslationRetargetingMode::Skeleton)
 					{
-						RootAtom.SetTranslation(SkeletonRemapping.RetargetBoneTranslationToTargetSkeleton(TargetSkeletonBoneIndex, RootAtom.GetTranslation()));
+						RootAtom.SetTranslation(SkeletonRemapping.RetargetBoneTranslationToTargetSkeleton(SkeletonBoneIndex, RootAtom.GetTranslation()));
 					}
 				}
 			}
@@ -478,7 +500,17 @@ void FDecompressionTools::DecompressPose(FLODPose& OutAnimationPoseData,
 					OutAnimationPoseData.LocalTransformsView[LODBoneIndex].SetRotation(SkeletonRemapping.RetargetAdditiveRotationToTargetSkeleton(TargetSkeletonBoneIndex, OutAnimationPoseData.LocalTransformsView[LODBoneIndex].GetRotation()));
 				}
 
-				if (TargetSkeleton->GetBoneTranslationRetargetingMode(TargetSkeletonBoneIndex, OutAnimationPoseData.GetDisableRetargeting()) != EBoneTranslationRetargetingMode::Skeleton)
+				// Check what retarget mode to use for the translational retargeting for this specific bone.
+				const int32 SourceSkeletonBoneIndex = SkeletonRemapping.GetSourceSkeletonBoneIndex(TargetSkeletonBoneIndex);
+				const EBoneTranslationRetargetingMode::Type RetargetMode = FAnimationRuntime::GetBoneTranslationRetargetingMode(
+					bUseSourceRetargetModes,
+					SourceSkeletonBoneIndex,
+					TargetSkeletonBoneIndex,
+					SourceSkeleton,
+					TargetSkeleton,
+					bDisableRetargeting);
+
+				if (RetargetMode != EBoneTranslationRetargetingMode::Skeleton)
 				{
 					OutAnimationPoseData.LocalTransformsView[LODBoneIndex].SetTranslation(SkeletonRemapping.RetargetAdditiveTranslationToTargetSkeleton(TargetSkeletonBoneIndex, OutAnimationPoseData.LocalTransformsView[LODBoneIndex].GetTranslation()));
 				}
@@ -490,7 +522,18 @@ void FDecompressionTools::DecompressPose(FLODPose& OutAnimationPoseData,
 			{
 				const int32 TargetSkeletonBoneIndex = LODBoneIndexToSkeletonBoneIndexMap[LODBoneIndex]; // ReferencePose.GetSkeletonBoneIndexFromLODBoneIndex(LODBoneIndex);
 				OutAnimationPoseData.LocalTransformsView[LODBoneIndex].SetRotation(SkeletonRemapping.RetargetBoneRotationToTargetSkeleton(TargetSkeletonBoneIndex, OutAnimationPoseData.LocalTransformsView[LODBoneIndex].GetRotation()));
-				if (TargetSkeleton->GetBoneTranslationRetargetingMode(TargetSkeletonBoneIndex, OutAnimationPoseData.GetDisableRetargeting()) != EBoneTranslationRetargetingMode::Skeleton)
+
+				// Check what retarget mode to use for the translational retargeting for this specific bone.
+				const int32 SourceSkeletonBoneIndex = SkeletonRemapping.GetSourceSkeletonBoneIndex(TargetSkeletonBoneIndex);
+				const EBoneTranslationRetargetingMode::Type RetargetMode = FAnimationRuntime::GetBoneTranslationRetargetingMode(
+					bUseSourceRetargetModes,
+					SourceSkeletonBoneIndex,
+					TargetSkeletonBoneIndex,
+					SourceSkeleton,
+					TargetSkeleton,
+					bDisableRetargeting);
+
+				if (RetargetMode != EBoneTranslationRetargetingMode::Skeleton)
 				{
 					OutAnimationPoseData.LocalTransformsView[LODBoneIndex].SetTranslation(SkeletonRemapping.RetargetBoneTranslationToTargetSkeleton(TargetSkeletonBoneIndex, OutAnimationPoseData.LocalTransformsView[LODBoneIndex].GetTranslation()));
 				}
@@ -653,7 +696,12 @@ void FDecompressionTools::RetargetBoneTransform(const FReferencePose& ReferenceP
 		const int32 TargetSkeletonBoneIndex = ReferencePose.GetSkeletonBoneIndexFromLODBoneIndex(LODBoneIndex);
 		const int32 SourceSkeletonBoneIndex = SkeletonRemapping.IsValid() ? SkeletonRemapping.GetSourceSkeletonBoneIndex(TargetSkeletonBoneIndex) : SkeletonBoneIndex;
 
-		switch (TargetSkeleton->GetBoneTranslationRetargetingMode(TargetSkeletonBoneIndex, bDisableRetargeting))
+		const bool UseSourceRetargetModes = TargetSkeleton->GetUseRetargetModesFromCompatibleSkeleton();
+		const EBoneTranslationRetargetingMode::Type RetargetMode = UseSourceRetargetModes 
+			? SourceSkeleton->GetBoneTranslationRetargetingMode(SourceSkeletonBoneIndex, bDisableRetargeting) 
+			: TargetSkeleton->GetBoneTranslationRetargetingMode(TargetSkeletonBoneIndex, bDisableRetargeting);
+
+		switch (RetargetMode)
 		{
 		case EBoneTranslationRetargetingMode::AnimationScaled:
 		{
diff --git a/Engine/Source/Editor/Persona/Private/SRetargetSources.cpp b/Engine/Source/Editor/Persona/Private/SRetargetSources.cpp
--- a/Engine/Source/Editor/Persona/Private/SRetargetSources.cpp
+++ b/Engine/Source/Editor/Persona/Private/SRetargetSources.cpp
@@ -8,6 +8,7 @@
 #include "Styling/AppStyle.h"
 #include "Widgets/Layout/SSeparator.h"
 #include "Widgets/Input/SButton.h"
+#include "Widgets/Input/SCheckBox.h"
 #include "Animation/DebugSkelMeshComponent.h"
 #include "Widgets/SToolTip.h"
 #include "IDocumentation.h"
@@ -17,6 +18,7 @@
 #include "IEditableSkeleton.h"
 #include "PropertyCustomizationHelpers.h"
 #include "Framework/MultiBox/MultiBoxBuilder.h"
+#include "Animation/Skeleton.h"
 
 
 #define LOCTEXT_NAMESPACE "SRetargetSources"
@@ -26,6 +28,14 @@ void SRetargetSources::Construct(
 	const TSharedRef<IEditableSkeleton>& InEditableSkeleton,
 	FSimpleMulticastDelegate& InOnPostUndo)
 {
+	const FText SourceRetargetModesToolTip = LOCTEXT("SourceRetargetModesCheckBoxTooltip", 
+		"Should we use the per bone translational retarget mode from the source (compatible) skeleton's instead of from this skeleton? On default this is disabled. "
+		"Enabling this would allow you to have one shared set of animations. You would configure the retarget settings on the animation skeleton. "
+		"Then every character that plays animations from this source skeleton will use the translational retarget settings from the source skeleton, which saves you from "
+		"having to configure the retarget modes for every bone in every character as they can be setup just once now on the animation skeleton.");
+
+	EditableSkeletonWeakPtr = InEditableSkeleton;
+
 	const FString DocLink = TEXT("Shared/Editors/Persona");
 	ChildSlot
 	[
@@ -39,7 +49,7 @@ void SRetargetSources::Construct(
 			.TextStyle(FAppStyle::Get(), "Persona.RetargetManager.ImportantText")
 			.Text(LOCTEXT("RetargetSource_Title", "Manage Retarget Sources"))
 		]
-		
+
 		+ SVerticalBox::Slot()
 		.Padding(5, 5)
 		.FillHeight(0.5)
@@ -65,6 +75,48 @@ void SRetargetSources::Construct(
 			.Text(LOCTEXT("CompatibleSkeletons_Title", "Manage Compatible Skeletons"))
 		]
 
+		+SVerticalBox::Slot()
+		.AutoHeight()
+		.VAlign(EVerticalAlignment::VAlign_Center)
+		.Padding(5, 5)
+		[
+			SNew(SHorizontalBox)
+			+SHorizontalBox::Slot()
+			.AutoWidth()
+			.VAlign(EVerticalAlignment::VAlign_Center)
+			[
+				SNew(SCheckBox)
+				.IsChecked_Lambda([this]() -> ECheckBoxState
+				{
+					if (EditableSkeletonWeakPtr.IsValid())
+					{
+						return EditableSkeletonWeakPtr.Pin()->GetSkeleton().GetUseRetargetModesFromCompatibleSkeleton() ? ECheckBoxState::Checked : ECheckBoxState::Unchecked;
+					}
+					return ECheckBoxState::Undetermined;
+				})
+				.OnCheckStateChanged_Lambda([this](ECheckBoxState NewState)
+				{
+					if (EditableSkeletonWeakPtr.IsValid())
+					{
+						const bool bIsChecked = (NewState == ECheckBoxState::Checked);
+						USkeleton& Skeleton = const_cast<USkeleton&>(EditableSkeletonWeakPtr.Pin()->GetSkeleton());
+						Skeleton.SetUseRetargetModesFromCompatibleSkeleton(bIsChecked);
+						Skeleton.Modify();
+					}
+				})
+				.ToolTipText(SourceRetargetModesToolTip)
+			]
+			+SHorizontalBox::Slot()
+			.AutoWidth()
+			.Padding(5, 0)
+			.VAlign(EVerticalAlignment::VAlign_Center)
+			[
+				SNew(STextBlock)
+				.Text(LOCTEXT("UseFromSource_Text", "Inherit Translation Retargeting"))
+				.ToolTipText(SourceRetargetModesToolTip)
+			]
+		]		
+
 		+ SVerticalBox::Slot()
 		.Padding(5, 5)
 		.FillHeight(0.5)
@@ -75,5 +127,4 @@ void SRetargetSources::Construct(
 	];
 }
 
-
 #undef LOCTEXT_NAMESPACE
diff --git a/Engine/Source/Editor/Persona/Private/SRetargetSources.h b/Engine/Source/Editor/Persona/Private/SRetargetSources.h
--- a/Engine/Source/Editor/Persona/Private/SRetargetSources.h
+++ b/Engine/Source/Editor/Persona/Private/SRetargetSources.h
@@ -23,4 +23,7 @@ public:
 		const FArguments& InArgs,
 		const TSharedRef<class IEditableSkeleton>& InEditableSkeleton,
 		FSimpleMulticastDelegate& InOnPostUndo);
+
+private:
+	TWeakPtr<class IEditableSkeleton> EditableSkeletonWeakPtr;
 };
diff --git a/Engine/Source/Runtime/Engine/Classes/Animation/Skeleton.h b/Engine/Source/Runtime/Engine/Classes/Animation/Skeleton.h
--- a/Engine/Source/Runtime/Engine/Classes/Animation/Skeleton.h
+++ b/Engine/Source/Runtime/Engine/Classes/Animation/Skeleton.h
@@ -344,6 +344,15 @@ protected:
 	UPROPERTY(EditAnywhere, BlueprintReadOnly, Category = CompatibleSkeletons)
 	TArray<TSoftObjectPtr<USkeleton>> CompatibleSkeletons;
 
+	/**
+	 * Should we use the per bone translational retarget mode from the source (compatible) skeleton's instead of from this skeleton? On default this is disabled.
+	 * Enabling this would allow you to have one shared set of animations. You would configure the retarget settings on the animation skeleton.
+	 * Then every character that plays animations from this source skeleton will use the translational retarget settings from the source skeleton, which saves you from
+	 * having to configure the retarget modes for every bone in every character as they can be setup just once now on the animation skeleton.
+	 */
+	UPROPERTY(EditAnywhere, Category = CompatibleSkeletons)
+	bool bUseRetargetModesFromCompatibleSkeleton = false;
+
 public:
 	//~ Begin UObject Interface.
 #if WITH_EDITOR
@@ -467,6 +476,9 @@ public:
 	UE_DEPRECATED(5.2, "Please use UE::Anim::FSkeletonRemappingRegistry::GetRemapping.")
 	ENGINE_API const FSkeletonRemapping* GetSkeletonRemapping(const USkeleton* SourceSkeleton) const;
 
+	ENGINE_API bool GetUseRetargetModesFromCompatibleSkeleton() const;
+	ENGINE_API void SetUseRetargetModesFromCompatibleSkeleton(bool bUse);
+
 #if WITH_EDITOR
 	// Get existing (seen) sync marker names for this Skeleton
 	const TArray<FName>& GetExistingMarkerNames() const { return ExistingMarkerNames; }
diff --git a/Engine/Source/Runtime/Engine/Private/Animation/AnimationDecompression.cpp b/Engine/Source/Runtime/Engine/Private/Animation/AnimationDecompression.cpp
--- a/Engine/Source/Runtime/Engine/Private/Animation/AnimationDecompression.cpp
+++ b/Engine/Source/Runtime/Engine/Private/Animation/AnimationDecompression.cpp
@@ -50,8 +50,9 @@ void DecompressPose(FCompactPose& OutPose,
 	const FBoneContainer& RequiredBones = OutPose.GetBoneContainer();
 	const int32 NumTracks = CompressedData.CompressedTrackToSkeletonMapTable.Num();
 
+	const USkeleton* SourceSkeleton = DecompressionContext.GetSourceSkeleton();
 	const USkeleton* TargetSkeleton = RequiredBones.GetSkeletonAsset();
-	const FSkeletonRemapping& SkeletonRemapping = UE::Anim::FSkeletonRemappingRegistry::Get().GetRemapping(DecompressionContext.GetSourceSkeleton(), TargetSkeleton);
+	const FSkeletonRemapping& SkeletonRemapping = UE::Anim::FSkeletonRemappingRegistry::Get().GetRemapping(SourceSkeleton, TargetSkeleton);
 
 	FGetBonePoseScratchArea& ScratchArea = FGetBonePoseScratchArea::Get();
 	BoneTrackArray& RotationScalePairs = ScratchArea.RotationScalePairs;
@@ -67,6 +68,9 @@ void DecompressPose(FCompactPose& OutPose,
 	AnimRelativeRetargetingPairs.Reset();
 	OrientAndScaleRetargetingPairs.Reset();
 
+	const bool bUseSourceRetargetModes = TargetSkeleton->GetUseRetargetModesFromCompatibleSkeleton();
+	const bool bDisableRetargeting = RequiredBones.GetDisableRetargeting();
+
 	const bool bIsMeshSpaceAdditive = DecompressionContext.GetAdditiveType() == AAT_RotationOffsetMeshSpace;
 	TBitArray<>& AnimatedCompactRotations = ScratchArea.AnimatedCompactRotations;
 	if (bIsMeshSpaceAdditive)
@@ -102,8 +106,17 @@ void DecompressPose(FCompactPose& OutPose,
 						AnimatedCompactRotations[CompactPoseBoneIndex] = true;
 					}
 
+					// Check what retarget mode to use for the translational retargeting for this specific bone.
+					const EBoneTranslationRetargetingMode::Type RetargetMode = FAnimationRuntime::GetBoneTranslationRetargetingMode(
+						bUseSourceRetargetModes,
+						SourceSkeletonBoneIndex,
+						TargetSkeletonBoneIndex,
+						SourceSkeleton,
+						TargetSkeleton,
+						bDisableRetargeting);
+
 					// Skip extracting translation component for EBoneTranslationRetargetingMode::Skeleton.
-					switch (TargetSkeleton->GetBoneTranslationRetargetingMode(TargetSkeletonBoneIndex, RequiredBones.GetDisableRetargeting()))
+					switch (RetargetMode)
 					{
 					case EBoneTranslationRetargetingMode::Animation:
 						TranslationPairs.Add(BoneTrackPair(CompactPoseBoneIndex, TrackIndex));
@@ -159,12 +172,21 @@ void DecompressPose(FCompactPose& OutPose,
 				// Root bone does not require fix-up for additive animations as there is no parent delta rotation to account for
 				if (!DecompressionContext.IsAdditiveAnimation())
 				{
-					const int32 TargetSkeletonBoneIndex = 0;
-
-					RootAtom.SetRotation(SkeletonRemapping.RetargetBoneRotationToTargetSkeleton(TargetSkeletonBoneIndex, RootAtom.GetRotation()));
-					if (TargetSkeleton->GetBoneTranslationRetargetingMode(TargetSkeletonBoneIndex, RequiredBones.GetDisableRetargeting()) != EBoneTranslationRetargetingMode::Skeleton)
+					constexpr int32 SkeletonBoneIndex = 0;
+
+					// Check what retarget mode to use for the translational retargeting for this specific bone.
+					const EBoneTranslationRetargetingMode::Type RetargetMode = FAnimationRuntime::GetBoneTranslationRetargetingMode(
+						bUseSourceRetargetModes,
+						SkeletonBoneIndex,
+						SkeletonBoneIndex,
+						SourceSkeleton,
+						TargetSkeleton,
+						bDisableRetargeting);
+
+					RootAtom.SetRotation(SkeletonRemapping.RetargetBoneRotationToTargetSkeleton(SkeletonBoneIndex, RootAtom.GetRotation()));
+					if (RetargetMode != EBoneTranslationRetargetingMode::Skeleton)
 					{
-						RootAtom.SetTranslation(SkeletonRemapping.RetargetBoneTranslationToTargetSkeleton(TargetSkeletonBoneIndex, RootAtom.GetTranslation()));
+						RootAtom.SetTranslation(SkeletonRemapping.RetargetBoneTranslationToTargetSkeleton(SkeletonBoneIndex, RootAtom.GetTranslation()));
 					}
 				}
 			}
@@ -195,8 +217,18 @@ void DecompressPose(FCompactPose& OutPose,
 				{
 					OutPose[BoneIndex].SetRotation(SkeletonRemapping.RetargetAdditiveRotationToTargetSkeleton(TargetSkeletonBoneIndex, OutPose[BoneIndex].GetRotation()));
 				}
-
-				if (TargetSkeleton->GetBoneTranslationRetargetingMode(TargetSkeletonBoneIndex, RequiredBones.GetDisableRetargeting()) != EBoneTranslationRetargetingMode::Skeleton)
+			
+				// Check what retarget mode to use for the translational retargeting for this specific bone.
+				const int32 SourceSkeletonBoneIndex = SkeletonRemapping.GetSourceSkeletonBoneIndex(TargetSkeletonBoneIndex);
+				const EBoneTranslationRetargetingMode::Type RetargetMode = FAnimationRuntime::GetBoneTranslationRetargetingMode(
+					bUseSourceRetargetModes,
+					SourceSkeletonBoneIndex,
+					TargetSkeletonBoneIndex,
+					SourceSkeleton,
+					TargetSkeleton,
+					bDisableRetargeting);
+
+				if (RetargetMode != EBoneTranslationRetargetingMode::Skeleton)
 				{
 					OutPose[BoneIndex].SetTranslation(SkeletonRemapping.RetargetAdditiveTranslationToTargetSkeleton(TargetSkeletonBoneIndex, OutPose[BoneIndex].GetTranslation()));
 				}
@@ -208,7 +240,18 @@ void DecompressPose(FCompactPose& OutPose,
 			{
 				const int32 TargetSkeletonBoneIndex = RequiredBones.GetSkeletonIndex(BoneIndex);
 				OutPose[BoneIndex].SetRotation(SkeletonRemapping.RetargetBoneRotationToTargetSkeleton(TargetSkeletonBoneIndex, OutPose[BoneIndex].GetRotation()));
-				if (TargetSkeleton->GetBoneTranslationRetargetingMode(TargetSkeletonBoneIndex, RequiredBones.GetDisableRetargeting()) != EBoneTranslationRetargetingMode::Skeleton)
+
+				// Check what retarget mode to use for the translational retargeting for this specific bone.
+				const int32 SourceSkeletonBoneIndex = SkeletonRemapping.GetSourceSkeletonBoneIndex(TargetSkeletonBoneIndex);
+				const EBoneTranslationRetargetingMode::Type RetargetMode = FAnimationRuntime::GetBoneTranslationRetargetingMode(
+					bUseSourceRetargetModes,
+					SourceSkeletonBoneIndex,
+					TargetSkeletonBoneIndex,
+					SourceSkeleton,
+					TargetSkeleton,
+					bDisableRetargeting);
+
+				if (RetargetMode != EBoneTranslationRetargetingMode::Skeleton)
 				{
 					OutPose[BoneIndex].SetTranslation(SkeletonRemapping.RetargetBoneTranslationToTargetSkeleton(TargetSkeletonBoneIndex, OutPose[BoneIndex].GetTranslation()));
 				}
diff --git a/Engine/Source/Runtime/Engine/Private/Animation/AnimationRuntime.cpp b/Engine/Source/Runtime/Engine/Private/Animation/AnimationRuntime.cpp
--- a/Engine/Source/Runtime/Engine/Private/Animation/AnimationRuntime.cpp
+++ b/Engine/Source/Runtime/Engine/Private/Animation/AnimationRuntime.cpp
@@ -2739,6 +2739,20 @@ void FAnimationRuntime::RetargetBoneTransform(const USkeleton* SourceSkeleton, c
 	}
 }
 
+EBoneTranslationRetargetingMode::Type FAnimationRuntime::GetBoneTranslationRetargetingMode(const bool bUseSourceRetargetModes, const int32 SourceSkeletonBoneIndex, const int32 TargetSkeletonBoneIndex, const USkeleton* SourceSkeleton, const USkeleton* TargetSkeleton, const bool bDisableRetargeting)
+{
+	if (bUseSourceRetargetModes)
+	{
+		return (SourceSkeletonBoneIndex != INDEX_NONE) 
+			? SourceSkeleton->GetBoneTranslationRetargetingMode(SourceSkeletonBoneIndex, bDisableRetargeting) 
+			: TargetSkeleton->GetBoneTranslationRetargetingMode(TargetSkeletonBoneIndex, bDisableRetargeting);
+	}
+	else
+	{
+		return TargetSkeleton->GetBoneTranslationRetargetingMode(TargetSkeletonBoneIndex, bDisableRetargeting);
+	}
+}
+
 void FAnimationRuntime::RetargetBoneTransform(const USkeleton* SourceSkeleton, const FName& SourceName, const TArray<FTransform>& RetargetTransforms, FTransform& BoneTransform, const int32 SkeletonBoneIndex, const FCompactPoseBoneIndex& BoneIndex, const FBoneContainer& RequiredBones, const bool bIsBakedAdditive)
 {
 	check(!RetargetTransforms.IsEmpty());
@@ -2758,8 +2772,18 @@ void FAnimationRuntime::RetargetBoneTransform(const USkeleton* SourceSkeleton, c
 			SourceSkeletonBoneIndex = SkeletonRemapping.GetSourceSkeletonBoneIndex(TargetSkeletonBoneIndex);
 			BoneTransform = SkeletonRemapping.RetargetBoneTransformToTargetSkeleton(TargetSkeletonBoneIndex, BoneTransform);
 		}
-		
-		switch (TargetSkeleton->GetBoneTranslationRetargetingMode(TargetSkeletonBoneIndex, RequiredBones.GetDisableRetargeting()))
+
+		// Check what retarget mode to use for the translational retargeting for this specific bone.
+		const bool bUseSourceRetargetModes = TargetSkeleton->GetUseRetargetModesFromCompatibleSkeleton();
+		const EBoneTranslationRetargetingMode::Type RetargetMode = FAnimationRuntime::GetBoneTranslationRetargetingMode(
+			bUseSourceRetargetModes,
+			SourceSkeletonBoneIndex,
+			TargetSkeletonBoneIndex,
+			SourceSkeleton,
+			TargetSkeleton,
+			RequiredBones.GetDisableRetargeting());
+			
+		switch (RetargetMode)
 		{
 			case EBoneTranslationRetargetingMode::AnimationScaled:
 			{
diff --git a/Engine/Source/Runtime/Engine/Private/Animation/Skeleton.cpp b/Engine/Source/Runtime/Engine/Private/Animation/Skeleton.cpp
--- a/Engine/Source/Runtime/Engine/Private/Animation/Skeleton.cpp
+++ b/Engine/Source/Runtime/Engine/Private/Animation/Skeleton.cpp
@@ -2374,5 +2374,15 @@ UAnimCurveMetaData* USkeleton::GetOrCreateCurveMetaDataObject()
 	return AnimCurveMetaData;
 }
 
+bool USkeleton::GetUseRetargetModesFromCompatibleSkeleton() const
+{
+	return bUseRetargetModesFromCompatibleSkeleton;
+}
+
+void USkeleton::SetUseRetargetModesFromCompatibleSkeleton(bool bUse)
+{
+	bUseRetargetModesFromCompatibleSkeleton = bUse;
+}
+
 #undef LOCTEXT_NAMESPACE 
 
diff --git a/Engine/Source/Runtime/Engine/Private/BoneContainer.cpp b/Engine/Source/Runtime/Engine/Private/BoneContainer.cpp
--- a/Engine/Source/Runtime/Engine/Private/BoneContainer.cpp
+++ b/Engine/Source/Runtime/Engine/Private/BoneContainer.cpp
@@ -8,6 +8,7 @@
 #include "Animation/AnimCurveTypes.h"
 #include "Animation/SkeletonRemappingRegistry.h"
 #include "Animation/SkeletonRemapping.h"
+#include "AnimationRuntime.h"
 
 LLM_DEFINE_TAG(BoneContainer);
 
@@ -444,7 +445,17 @@ const FRetargetSourceCachedData& FBoneContainer::GetRetargetSourceCachedData(
 			const int32 TargetSkeletonBoneIndex = CompactPoseToSkeletonIndex[CompactBoneIndex];
 			const int32 SourceSkeletonBoneIndex = InRemapping.IsValid() ? InRemapping.GetSourceSkeletonBoneIndex(TargetSkeletonBoneIndex) : TargetSkeletonBoneIndex;
 
-			if (AssetSkeleton.GetEvenIfUnreachable()->GetBoneTranslationRetargetingMode(TargetSkeletonBoneIndex, bDisableRetargeting) == EBoneTranslationRetargetingMode::OrientAndScale)
+			// Check what retarget mode to use for the translational retargeting for this specific bone.
+			const bool bUseSourceRetargetModes = AssetSkeleton.GetEvenIfUnreachable()->GetUseRetargetModesFromCompatibleSkeleton();
+			const EBoneTranslationRetargetingMode::Type RetargetMode = FAnimationRuntime::GetBoneTranslationRetargetingMode(
+				bUseSourceRetargetModes,
+				SourceSkeletonBoneIndex,
+				TargetSkeletonBoneIndex,
+				SourceSkeleton,
+				AssetSkeleton.GetEvenIfUnreachable(),
+				bDisableRetargeting);
+
+			if (RetargetMode == EBoneTranslationRetargetingMode::OrientAndScale)
 			{
 				if(AuthoredOnRefSkeleton.IsValidIndex(SourceSkeletonBoneIndex))
 				{
diff --git a/Engine/Source/Runtime/Engine/Public/AnimationRuntime.h b/Engine/Source/Runtime/Engine/Public/AnimationRuntime.h
--- a/Engine/Source/Runtime/Engine/Public/AnimationRuntime.h
+++ b/Engine/Source/Runtime/Engine/Public/AnimationRuntime.h
@@ -609,6 +609,21 @@ public:
 	static ENGINE_API void FillUpComponentSpaceTransforms(const FReferenceSkeleton& RefSkeleton, const TArrayView<const FTransform> &BoneSpaceTransforms, TArray<FTransform> &ComponentSpaceTransforms);
 	static ENGINE_API void MakeSkeletonRefPoseFromMesh(const USkeletalMesh* InMesh, const USkeleton* InSkeleton, TArray<FTransform>& OutBoneBuffer);
 
+	/**
+	 * Get the translational retargeting mode for a given bone on the target skeleton when sampling a given source animation.
+	 * If bUseSourceRetargetModes is set to true, and SourceSkeletonBoneIndex is INDEX_NONE, then the target bone and target skeleton are used instead.
+	 * If bUseSourceRetargetModes is set to false, the SourceSkeleton and SourceSkeletonBoneIndex values are ignored and we get the retarget mode directly from the target skeleton using the target bone index.
+	 * 
+	 * @param	bUseSourceRetargetModes		Set to true when you wish to extract the retarget mode from the source skeleton instead of the target skeleton.
+	 * @param	SourceSkeletonBoneIndex		The skeleton bone index of the bone on the source skeleton to get the retarget mode from.
+	 * @param	TargetSkeletonBoneIndex		The target bone index of the bone on the target skeleton to get the retarget mode from.
+	 * @param	SourceSkeleton				A pointer to the skeleton of the source (the animation's skeleton for example).
+	 * @param	TargetSkeleton				A pointer to the target skeleton, which is the skeleton to apply transforms to.
+	 * @param	bDisableRetargeting			Is translational retargeting disabled?
+	 * @return	The retargeting mode for the specific bone.
+	 */
+	static ENGINE_API EBoneTranslationRetargetingMode::Type GetBoneTranslationRetargetingMode(const bool bUseSourceRetargetModes, const int32 SourceSkeletonBoneIndex, const int32 TargetSkeletonBoneIndex, const USkeleton* SourceSkeleton, const USkeleton* TargetSkeleton, const bool bDisableRetargeting);
+
 	/**
 	 * Calculate the component-space bone transform for the specified bone. This is similar to GetComponentSpaceTransform(), but uses a cached array of 
 	 * transforms to prevent re-calculations when requesting transforms for multiple bones with shared parents. This is only likely to be more efficient than
