AnimNext: Module Components can now be used to marshal data from the game thread

Two initial examples:
- FRigUnit_AnimNextWriteSkeletalMeshComponentPose now does not require a skeletal mesh component. If a component is not supplied, then it will interrogate the actor to find the first one and use that. This module component is 'static', in that it wont continue to update the component if the component is removed and a new one re-added (this is mostly because there arent actor-level hooks for this - we could add a mechanism to the skeletal mesh component if we wanted this).
- FRigUnit_GetActorTransform is now thread safe. On initialization, it hooks into the actor root component's TransformUpdated delegate, then forwards new transforms to the module via queued tasks

diff --git a/Engine/Plugins/Experimental/AnimNext/Source/AnimNext/Internal/Component/AnimNextWorldSubsystem.h b/Engine/Plugins/Experimental/AnimNext/Source/AnimNext/Internal/Component/AnimNextWorldSubsystem.h
--- a/Engine/Plugins/Experimental/AnimNext/Source/AnimNext/Internal/Component/AnimNextWorldSubsystem.h
+++ b/Engine/Plugins/Experimental/AnimNext/Source/AnimNext/Internal/Component/AnimNextWorldSubsystem.h
@@ -9,6 +9,7 @@
 #include "Module/ModuleTickFunction.h"
 #include "Module/AnimNextModuleInstance.h"
 #include "Module/ModuleTaskContext.h"
+#include "Module/TaskRunLocation.h"
 #include "AnimNextWorldSubsystem.generated.h"
 
 class UAnimNextComponent;
diff --git a/Engine/Plugins/Experimental/AnimNext/Source/AnimNext/Internal/Module/AnimNextActorComponentReferenceComponent.h b/Engine/Plugins/Experimental/AnimNext/Source/AnimNext/Internal/Module/AnimNextActorComponentReferenceComponent.h
new file mode 100644
--- /dev/null
+++ b/Engine/Plugins/Experimental/AnimNext/Source/AnimNext/Internal/Module/AnimNextActorComponentReferenceComponent.h
@@ -0,0 +1,31 @@
+﻿// Copyright Epic Games, Inc. All Rights Reserved.
+
+#pragma once
+
+#include "Module/AnimNextModuleInstanceComponent.h"
+#include "Templates/SubclassOf.h"
+
+#include "AnimNextActorComponentReferenceComponent.generated.h"
+
+class UActorComponent;
+
+// Module instance component used to provide a reference to an actor component
+// Note that this does not support dynamic component updates (e.g. if we have multiple, or changing components). It grabs a ptr to the component on
+// initialization and does not update it.
+USTRUCT()
+struct FAnimNextActorComponentReferenceComponent : public FAnimNextModuleInstanceComponent
+{
+	GENERATED_BODY()
+
+protected:
+	// Helper function used by derived classes to initialize the component
+	void OnInitializeHelper(UScriptStruct* InScriptStruct);
+
+	// The component we reference
+	UPROPERTY()
+	TObjectPtr<UActorComponent> Component;
+
+	// The type of the component to find
+	UPROPERTY()
+	TSubclassOf<UActorComponent> ComponentType;
+};
\ No newline at end of file
diff --git a/Engine/Plugins/Experimental/AnimNext/Source/AnimNext/Internal/Module/AnimNextModuleInstance.h b/Engine/Plugins/Experimental/AnimNext/Source/AnimNext/Internal/Module/AnimNextModuleInstance.h
--- a/Engine/Plugins/Experimental/AnimNext/Source/AnimNext/Internal/Module/AnimNextModuleInstance.h
+++ b/Engine/Plugins/Experimental/AnimNext/Source/AnimNext/Internal/Module/AnimNextModuleInstance.h
@@ -11,7 +11,7 @@
 #include "TraitCore/TraitEvent.h"
 #include "TraitCore/TraitEventList.h"
 #include "Module/ModuleHandle.h"
-#include "Module/ModuleInstanceComponent.h"
+#include "Module/AnimNextModuleInstanceComponent.h"
 #include "StructUtils/InstancedStruct.h"
 #include "Module/AnimNextModuleInitMethod.h"
 
@@ -130,7 +130,7 @@ struct FAnimNextModuleInstance : public FAnimNextDataInterfaceInstance
 #endif
 
 	// Run a simple task on the GT via FFunctionGraphTask::CreateAndDispatchWhenReady
-	static void RunTaskOnGameThread(TUniqueFunction<void(void)>&& InFunction);
+	ANIMNEXT_API static void RunTaskOnGameThread(TUniqueFunction<void(void)>&& InFunction);
 
 	// Get the names of all available module names
 	static void GetModuleEventNames(TArray<FString>& OutNames, bool bIncludeInternal = false);
@@ -151,6 +151,9 @@ private:
 	// Setup the entry
 	void Initialize();
 
+	// Tear-down the entry
+	void Uninitialize();
+
 	// Clear data that binds the schedule to a runtime (e.g. tick functions) and any instance data
 	void ResetBindingsAndInstanceData();
 
diff --git a/Engine/Plugins/Experimental/AnimNext/Source/AnimNext/Internal/Module/AnimNextModuleInstanceComponent.h b/Engine/Plugins/Experimental/AnimNext/Source/AnimNext/Internal/Module/AnimNextModuleInstanceComponent.h
new file mode 100644
--- /dev/null
+++ b/Engine/Plugins/Experimental/AnimNext/Source/AnimNext/Internal/Module/AnimNextModuleInstanceComponent.h
@@ -0,0 +1,52 @@
+// Copyright Epic Games, Inc. All Rights Reserved.
+
+#pragma once
+
+#include "CoreMinimal.h"
+#include "AnimNextModuleInstanceComponent.generated.h"
+
+struct FAnimNextModuleInstance;
+struct FAnimNextTraitEvent;
+
+/**
+ * FAnimNextModuleInstanceComponent
+ *
+ * A module instance component is attached and owned by a module instance.
+ * It persists as long as it is needed.
+ */
+USTRUCT()
+struct FAnimNextModuleInstanceComponent
+{
+	GENERATED_BODY()
+
+	FAnimNextModuleInstanceComponent() = default;
+	virtual ~FAnimNextModuleInstanceComponent() {}
+
+	ANIMNEXT_API void Initialize(FAnimNextModuleInstance& InOwnerInstance);
+	ANIMNEXT_API void Uninitialize();
+
+	// Returns the owning module instance this component lives on
+	FAnimNextModuleInstance& GetModuleInstance() { return *OwnerInstance; }
+
+	// Returns the owning module instance this component lives on
+	const FAnimNextModuleInstance& GetModuleInstance() const { return *OwnerInstance; }
+
+	// Called when the component is first created to initialize it.
+	// This can occur on module initialize or lazily during execution.
+	virtual void OnInitialize() {}
+
+	// Called when the component is destroyed. Once created, components persist until the module instance is destroyed.
+	virtual void OnUninitialize() {}
+
+	// Called during module execution for any events to be handled
+	virtual void OnTraitEvent(FAnimNextTraitEvent& Event) {}
+
+	// Called at end of module execution each frame
+	virtual void OnEndExecution(float InDeltaTime) {}
+
+private:
+	// The owning module instance this component lives on
+	FAnimNextModuleInstance* OwnerInstance = nullptr;
+
+	friend struct FAnimNextModuleInstance;
+};
diff --git a/Engine/Plugins/Experimental/AnimNext/Source/AnimNext/Internal/Module/AnimNextSkeletalMeshComponentReferenceComponent.h b/Engine/Plugins/Experimental/AnimNext/Source/AnimNext/Internal/Module/AnimNextSkeletalMeshComponentReferenceComponent.h
new file mode 100644
--- /dev/null
+++ b/Engine/Plugins/Experimental/AnimNext/Source/AnimNext/Internal/Module/AnimNextSkeletalMeshComponentReferenceComponent.h
@@ -0,0 +1,21 @@
+﻿// Copyright Epic Games, Inc. All Rights Reserved.
+
+#pragma once
+
+#include "Module/AnimNextActorComponentReferenceComponent.h"
+
+#include "AnimNextSkeletalMeshComponentReferenceComponent.generated.h"
+
+// Module instance component used to fetch data about the current actor's skeletal mesh component
+USTRUCT()
+struct FAnimNextSkeletalMeshComponentReferenceComponent : public FAnimNextActorComponentReferenceComponent
+{
+	GENERATED_BODY()
+
+	FAnimNextSkeletalMeshComponentReferenceComponent();
+	USkeletalMeshComponent* GetComponent() const;
+
+private:
+	// FAnimNextModuleInstanceComponent interface
+	virtual void OnInitialize() override;
+};
diff --git a/Engine/Plugins/Experimental/AnimNext/Source/AnimNext/Internal/Module/ModuleInstanceComponent.h b/Engine/Plugins/Experimental/AnimNext/Source/AnimNext/Internal/Module/ModuleInstanceComponent.h
deleted file mode 100644
diff --git a/Engine/Plugins/Experimental/AnimNext/Source/AnimNext/Internal/Module/ModuleTaskContext.h b/Engine/Plugins/Experimental/AnimNext/Source/AnimNext/Internal/Module/ModuleTaskContext.h
deleted file mode 100644
diff --git a/Engine/Plugins/Experimental/AnimNext/Source/AnimNext/Private/Component/AnimNextComponent.cpp b/Engine/Plugins/Experimental/AnimNext/Source/AnimNext/Private/Component/AnimNextComponent.cpp
--- a/Engine/Plugins/Experimental/AnimNext/Source/AnimNext/Private/Component/AnimNextComponent.cpp
+++ b/Engine/Plugins/Experimental/AnimNext/Source/AnimNext/Private/Component/AnimNextComponent.cpp
@@ -885,6 +885,16 @@ void UAnimNextComponent::ShowDebugDrawing(bool bInShowDebugDrawing)
 #endif
 }
 
+void UAnimNextComponent::QueueTask(FName InModuleEventName, TUniqueFunction<void(const UE::AnimNext::FModuleTaskContext&)>&& InTaskFunction, UE::AnimNext::ETaskRunLocation InLocation)
+{
+	using namespace UE::AnimNext;
+
+	if(Subsystem)
+	{
+		Subsystem->QueueTask(this, InModuleEventName, MoveTemp(InTaskFunction), InLocation);
+	}
+}
+
 void UAnimNextComponent::QueueInputTraitEvent(FAnimNextTraitEventPtr Event)
 {
 	using namespace UE::AnimNext;
diff --git a/Engine/Plugins/Experimental/AnimNext/Source/AnimNext/Private/Component/AnimNextWorldSubsystem.cpp b/Engine/Plugins/Experimental/AnimNext/Source/AnimNext/Private/Component/AnimNextWorldSubsystem.cpp
--- a/Engine/Plugins/Experimental/AnimNext/Source/AnimNext/Private/Component/AnimNextWorldSubsystem.cpp
+++ b/Engine/Plugins/Experimental/AnimNext/Source/AnimNext/Private/Component/AnimNextWorldSubsystem.cpp
@@ -107,6 +107,8 @@ void UAnimNextWorldSubsystem::FlushPendingActions()
 			case FModulePendingAction::EType::ReleaseHandle:
 				if (IsValidHandle(PendingAction.Handle))
 				{
+					FAnimNextModuleInstance& Instance = Instances.Get(PendingAction.Handle);
+					Instance.Uninitialize();
 					Instances.Release(PendingAction.Handle);
 				}
 				break;
diff --git a/Engine/Plugins/Experimental/AnimNext/Source/AnimNext/Private/Graph/RigUnit_AnimNextWriteSkeletalMeshComponentPose.cpp b/Engine/Plugins/Experimental/AnimNext/Source/AnimNext/Private/Graph/RigUnit_AnimNextWriteSkeletalMeshComponentPose.cpp
--- a/Engine/Plugins/Experimental/AnimNext/Source/AnimNext/Private/Graph/RigUnit_AnimNextWriteSkeletalMeshComponentPose.cpp
+++ b/Engine/Plugins/Experimental/AnimNext/Source/AnimNext/Private/Graph/RigUnit_AnimNextWriteSkeletalMeshComponentPose.cpp
@@ -7,6 +7,7 @@
 #include "Component/SkinnedMeshComponentExtensions.h"
 #include "Components/SkeletalMeshComponent.h"
 #include "Engine/SkeletalMesh.h"
+#include "Module/AnimNextSkeletalMeshComponentReferenceComponent.h"
 
 #include UE_INLINE_GENERATED_CPP_BY_NAME(RigUnit_AnimNextWriteSkeletalMeshComponentPose)
 
@@ -18,36 +19,47 @@ FRigUnit_AnimNextWriteSkeletalMeshComponentPose_Execute()
 
 	using namespace UE::AnimNext;
 
-	if(SkeletalMeshComponent == nullptr)
+	USkeletalMeshComponent* OutputComponent = SkeletalMeshComponent;
+
+	// Defer to module component if no component is supplied
+	if(OutputComponent == nullptr)
+	{
+		const FAnimNextModuleContextData& ModuleContextData = ExecuteContext.GetContextData<FAnimNextModuleContextData>();
+		FAnimNextModuleInstance& ModuleInstance = ModuleContextData.GetModuleInstance();
+		const FAnimNextSkeletalMeshComponentReferenceComponent& ComponentReference = ModuleInstance.GetComponent<FAnimNextSkeletalMeshComponentReferenceComponent>();
+		OutputComponent = ComponentReference.GetComponent();
+	}
+
+	if(OutputComponent == nullptr)
 	{
 		return;
 	}
-	
+
 	if(!Pose.LODPose.IsValid())
 	{
 		return;
 	}
 
-	USkeletalMesh* SkeletalMesh = SkeletalMeshComponent->GetSkeletalMeshAsset();
+	USkeletalMesh* SkeletalMesh = OutputComponent->GetSkeletalMeshAsset();
 	if(SkeletalMesh == nullptr)
 	{
 		return;
 	}
 
 	// We cannot write to the skeletal mesh component if it is using anim BP
-	if(SkeletalMeshComponent->bEnableAnimation)
+	if(OutputComponent->bEnableAnimation)
 	{
 		UE_LOG(LogAnimation, Warning, TEXT("AnimNext: Could not write to skeletal mesh component - bEnableAnimation is true"))
 		return;
 	}
 
-	FDataHandle RefPoseHandle = FDataRegistry::Get()->GetOrGenerateReferencePose(SkeletalMeshComponent);
+	FDataHandle RefPoseHandle = FDataRegistry::Get()->GetOrGenerateReferencePose(OutputComponent);
 	const UE::AnimNext::FReferencePose& RefPose = RefPoseHandle.GetRef<UE::AnimNext::FReferencePose>();
 
 	FMemMark MemMark(FMemStack::Get());
 
 	// TODO: The LOD level we got here is not guaranteed to be stable. We should find another way to get/bind the skeletal mesh component's LOD level.
-	const int32 SkeletalMeshComponentLODLevel = SkeletalMeshComponent->GetPredictedLODLevel();
+	const int32 SkeletalMeshComponentLODLevel = OutputComponent->GetPredictedLODLevel();
 	const int32 PoseLODLevel = Pose.LODPose.LODLevel;
 
 	TArray<FTransform, TMemStackAllocator<>> LocalSpaceTransforms;
@@ -74,14 +86,14 @@ FRigUnit_AnimNextWriteSkeletalMeshComponentPose_Execute()
 	// Map LOD pose into local-space scratch buffer
 	FGenerationTools::RemapPose(Pose.LODPose, LocalSpaceTransforms);
 
-	SkeletalMeshComponent->AnimCurves.CopyFrom(Pose.Curves);
+	OutputComponent->AnimCurves.CopyFrom(Pose.Curves);
 
 	// Attributes require remapping since the indices are LOD indices and we want mesh indices
-	FGenerationTools::RemapAttributes(Pose.LODPose, Pose.Attributes, SkeletalMeshComponent->GetEditableCustomAttributes());
+	FGenerationTools::RemapAttributes(Pose.LODPose, Pose.Attributes, OutputComponent->GetEditableCustomAttributes());
 
 	// Convert and dispatch to renderer
 	UE::Anim::FSkinnedMeshComponentExtensions::CompleteAndDispatch(
-		SkeletalMeshComponent,
+		OutputComponent,
 		RefPose.GetMeshBoneIndexToParentMeshBoneIndexMap(),
 		RefPose.GetLODBoneIndexToMeshBoneIndexMap(SkeletalMeshComponentLODLevel), // Based on the LOD level of the skeletal mesh component
 		LocalSpaceTransforms);
diff --git a/Engine/Plugins/Experimental/AnimNext/Source/AnimNext/Private/Graph/RigUnit_AnimNextWriteSkeletalMeshComponentPose.h b/Engine/Plugins/Experimental/AnimNext/Source/AnimNext/Private/Graph/RigUnit_AnimNextWriteSkeletalMeshComponentPose.h
--- a/Engine/Plugins/Experimental/AnimNext/Source/AnimNext/Private/Graph/RigUnit_AnimNextWriteSkeletalMeshComponentPose.h
+++ b/Engine/Plugins/Experimental/AnimNext/Source/AnimNext/Private/Graph/RigUnit_AnimNextWriteSkeletalMeshComponentPose.h
@@ -9,7 +9,7 @@
 #include "RigUnit_AnimNextWriteSkeletalMeshComponentPose.generated.h"
 
 /** Writes a pose to a skeletal mesh component */
-USTRUCT(meta=(DisplayName="Write Pose", Category="Animation Graph", NodeColor="0, 1, 1", Keywords="Output,Pose,Port"))
+USTRUCT(meta=(DisplayName="Write Pose", Category="Animation Graph", NodeColor="0, 1, 1", Keywords="Output,Pose,Port", RequiredComponents="AnimNextSkeletalMeshComponentReferenceComponent"))
 struct FRigUnit_AnimNextWriteSkeletalMeshComponentPose : public FRigUnit_AnimNextBase
 {
 	GENERATED_BODY()
@@ -23,7 +23,7 @@ struct FRigUnit_AnimNextWriteSkeletalMeshComponentPose : public FRigUnit_AnimNex
 	UPROPERTY(EditAnywhere, Category = "Graph", meta = (Input))
 	FAnimNextGraphLODPose Pose;
 
-	// Mesh to write to
+	// Mesh to write to. If this is not supplied, then the first skeletal mesh component of the current actor will be used
 	UPROPERTY(EditAnywhere, Category = "Graph", meta = (Input))
 	TObjectPtr<USkeletalMeshComponent> SkeletalMeshComponent;
 
diff --git a/Engine/Plugins/Experimental/AnimNext/Source/AnimNext/Private/Graph/RigUnit_GetActorTransform.cpp b/Engine/Plugins/Experimental/AnimNext/Source/AnimNext/Private/Graph/RigUnit_GetActorTransform.cpp
--- a/Engine/Plugins/Experimental/AnimNext/Source/AnimNext/Private/Graph/RigUnit_GetActorTransform.cpp
+++ b/Engine/Plugins/Experimental/AnimNext/Source/AnimNext/Private/Graph/RigUnit_GetActorTransform.cpp
@@ -4,28 +4,57 @@
 
 #include "Component/AnimNextComponent.h"
 
+void FAnimNextActorTransformComponent::OnInitialize()
+{
+	FAnimNextModuleInstance& ModuleInstance = GetModuleInstance();
+	UAnimNextComponent* AnimNextComponent = Cast<UAnimNextComponent>(ModuleInstance.GetObject());
+
+	FAnimNextModuleInstance::RunTaskOnGameThread(
+		[WeakAnimNextComponent = TWeakObjectPtr<UAnimNextComponent>(AnimNextComponent)]()
+		{
+			UAnimNextComponent* Component = WeakAnimNextComponent.Get();
+			if(Component == nullptr)
+			{
+				return;
+			}
+
+			AActor* Owner = Component->GetOwner();
+			if(Owner == nullptr)
+			{
+				return;
+			}
+
+			USceneComponent* SceneComponent = Owner->GetRootComponent();
+			if(SceneComponent == nullptr)
+			{
+				return;
+			}
+
+			SceneComponent->TransformUpdated.AddWeakLambda(Component, [Component](USceneComponent* UpdatedComponent, EUpdateTransformFlags UpdateTransformFlags, ETeleportType Teleport)
+			{
+				AActor* UpdatedComponentOwner = UpdatedComponent->GetOwner();
+				if(UpdatedComponentOwner == nullptr)
+				{
+					return;
+				}
+
+				Component->QueueTask(NAME_None, [Transform = UpdatedComponentOwner->GetActorTransform()](const UE::AnimNext::FModuleTaskContext& InContext)
+				{
+					InContext.TryAccessComponent<FAnimNextActorTransformComponent>([&Transform](FAnimNextActorTransformComponent& InComponent)
+					{
+						InComponent.ActorTransform = Transform;
+					});
+				});
+			});
+		});
+}
+
 FRigUnit_GetActorTransform_Execute()
 {
 	using namespace UE::AnimNext;
 
 	const FAnimNextModuleContextData& ModuleContextData = ExecuteContext.GetContextData<FAnimNextModuleContextData>();
 	FAnimNextModuleInstance& ModuleInstance = ModuleContextData.GetModuleInstance();
-
-	// TODO: this is not thread safe at the moment. Once we get the 'weak semantics' CL checked in, we can move commonly-used debug info
-	// into a module component that can be used to access GT data where needed. 
-	UAnimNextComponent* Component = Cast<UAnimNextComponent>(ModuleInstance.GetObject());
-	if(Component == nullptr)
-	{
-		Transform = FTransform::Identity;
-		return;
-	}
-
-	AActor* Owner = Component->GetOwner();
-	if(Owner == nullptr)
-	{
-		Transform = FTransform::Identity;
-		return;
-	}
-
-	Transform = Owner->GetActorTransform();
+	const FAnimNextActorTransformComponent& TransformComponent = ModuleInstance.GetComponent<FAnimNextActorTransformComponent>();
+	Transform = TransformComponent.GetActorTransform();
 }
\ No newline at end of file
diff --git a/Engine/Plugins/Experimental/AnimNext/Source/AnimNext/Private/Graph/RigUnit_GetActorTransform.h b/Engine/Plugins/Experimental/AnimNext/Source/AnimNext/Private/Graph/RigUnit_GetActorTransform.h
--- a/Engine/Plugins/Experimental/AnimNext/Source/AnimNext/Private/Graph/RigUnit_GetActorTransform.h
+++ b/Engine/Plugins/Experimental/AnimNext/Source/AnimNext/Private/Graph/RigUnit_GetActorTransform.h
@@ -1,14 +1,32 @@
-﻿// Copyright Epic Games, Inc. All Rights Reserved.
+// Copyright Epic Games, Inc. All Rights Reserved.
 
 #pragma once
 
 #include "AnimNextExecuteContext.h"
 #include "Graph/RigUnit_AnimNextBase.h"
+#include "Module/AnimNextModuleInstanceComponent.h"
 
 #include "RigUnit_GetActorTransform.generated.h"
 
-/** Gets the transform of the host of this AnimNext module */
-USTRUCT(meta=(DisplayName="Get Host Transform", Category="Animation Graph", NodeColor="0, 1, 1", Keywords="Actor, Transform"))
+// Module component used to marshal actor transform data for RigVM units to use
+USTRUCT()
+struct FAnimNextActorTransformComponent : public FAnimNextModuleInstanceComponent
+{
+	GENERATED_BODY()
+
+	// Get the transform of the actor we are bound to
+	const FTransform& GetActorTransform() const { return ActorTransform; }
+
+private:
+	// FAnimNextModuleInstanceComponent interface
+	virtual void OnInitialize() override;
+
+	// Cached actor transform
+	FTransform ActorTransform = FTransform::Identity;
+};
+
+/** Gets the transform of the actor hosting this AnimNext module */
+USTRUCT(meta=(DisplayName="Get Host Transform", Category="Animation Graph", NodeColor="0, 1, 1", Keywords="Actor, Transform", RequiredComponents="AnimNextActorTransformComponent"))
 struct FRigUnit_GetActorTransform : public FRigUnit_AnimNextBase
 {
 	GENERATED_BODY()
@@ -20,5 +38,5 @@ struct FRigUnit_GetActorTransform : public FRigUnit_AnimNextBase
 
 	// The transform of our host actor
 	UPROPERTY(EditAnywhere, Category = "Graph", meta = (Output))
-	FTransform Transform;
+	FTransform Transform = FTransform::Identity;
 };
diff --git a/Engine/Plugins/Experimental/AnimNext/Source/AnimNext/Private/Graph/RigUnit_MakeReferencePoseFromSkeletalMeshComponent.cpp b/Engine/Plugins/Experimental/AnimNext/Source/AnimNext/Private/Graph/RigUnit_MakeReferencePoseFromSkeletalMeshComponent.cpp
--- a/Engine/Plugins/Experimental/AnimNext/Source/AnimNext/Private/Graph/RigUnit_MakeReferencePoseFromSkeletalMeshComponent.cpp
+++ b/Engine/Plugins/Experimental/AnimNext/Source/AnimNext/Private/Graph/RigUnit_MakeReferencePoseFromSkeletalMeshComponent.cpp
@@ -6,6 +6,7 @@
 #include "GenerationTools.h"
 #include "Components/SkeletalMeshComponent.h"
 #include "Engine/SkeletalMesh.h"
+#include "Module/AnimNextSkeletalMeshComponentReferenceComponent.h"
 
 #include UE_INLINE_GENERATED_CPP_BY_NAME(RigUnit_MakeReferencePoseFromSkeletalMeshComponent)
 
@@ -17,10 +18,21 @@ FRigUnit_MakeReferencePoseFromSkeletalMeshComponent_Execute()
 
 	using namespace UE::AnimNext;
 
-	if(SkeletalMeshComponent == nullptr)
+	USkeletalMeshComponent* InputComponent = SkeletalMeshComponent;
+
+	// Defer to module component if no component is supplied
+	if(InputComponent == nullptr)
+	{
+		const FAnimNextModuleContextData& ModuleContextData = ExecuteContext.GetContextData<FAnimNextModuleContextData>();
+		FAnimNextModuleInstance& ModuleInstance = ModuleContextData.GetModuleInstance();
+		const FAnimNextSkeletalMeshComponentReferenceComponent& ComponentReference = ModuleInstance.GetComponent<FAnimNextSkeletalMeshComponentReferenceComponent>();
+		InputComponent = ComponentReference.GetComponent();
+	}
+
+	if(InputComponent == nullptr)
 	{
 		return;
 	}
 
-	ReferencePose.ReferencePose = FDataRegistry::Get()->GetOrGenerateReferencePose(SkeletalMeshComponent);
+	ReferencePose.ReferencePose = FDataRegistry::Get()->GetOrGenerateReferencePose(InputComponent);
 }
diff --git a/Engine/Plugins/Experimental/AnimNext/Source/AnimNext/Private/Graph/RigUnit_MakeReferencePoseFromSkeletalMeshComponent.h b/Engine/Plugins/Experimental/AnimNext/Source/AnimNext/Private/Graph/RigUnit_MakeReferencePoseFromSkeletalMeshComponent.h
--- a/Engine/Plugins/Experimental/AnimNext/Source/AnimNext/Private/Graph/RigUnit_MakeReferencePoseFromSkeletalMeshComponent.h
+++ b/Engine/Plugins/Experimental/AnimNext/Source/AnimNext/Private/Graph/RigUnit_MakeReferencePoseFromSkeletalMeshComponent.h
@@ -9,7 +9,7 @@
 #include "RigUnit_MakeReferencePoseFromSkeletalMeshComponent.generated.h"
 
 /** Makes a reference pose from a skeletal mesh component */
-USTRUCT(meta=(DisplayName="Make Reference Pose", Category="Animation Graph", NodeColor="0, 1, 1", Keywords="Output,Pose,Port"))
+USTRUCT(meta=(DisplayName="Make Reference Pose", Category="Animation Graph", NodeColor="0, 1, 1", Keywords="Output,Pose,Port", RequiredComponents="AnimNextSkeletalMeshComponentReferenceComponent"))
 struct FRigUnit_MakeReferencePoseFromSkeletalMeshComponent : public FRigUnit_AnimNextBase
 {
 	GENERATED_BODY()
@@ -23,7 +23,7 @@ struct FRigUnit_MakeReferencePoseFromSkeletalMeshComponent : public FRigUnit_Ani
 	UPROPERTY(EditAnywhere, Category = "Graph", meta = (Input, Output))
 	FAnimNextGraphReferencePose ReferencePose;
 
-	// Mesh to use to generate the reference pose
+	// Mesh to use to generate the reference pose. If this is not supplied, then the first skeletal mesh component of the current actor will be used
 	UPROPERTY(EditAnywhere, Category = "Graph", meta = (Input))
 	TObjectPtr<USkeletalMeshComponent> SkeletalMeshComponent;
 
diff --git a/Engine/Plugins/Experimental/AnimNext/Source/AnimNext/Private/Module/AnimNextActorComponentReferenceComponent.cpp b/Engine/Plugins/Experimental/AnimNext/Source/AnimNext/Private/Module/AnimNextActorComponentReferenceComponent.cpp
new file mode 100644
--- /dev/null
+++ b/Engine/Plugins/Experimental/AnimNext/Source/AnimNext/Private/Module/AnimNextActorComponentReferenceComponent.cpp
@@ -0,0 +1,39 @@
+﻿// Copyright Epic Games, Inc. All Rights Reserved.
+
+#include "Module/AnimNextActorComponentReferenceComponent.h"
+
+#include "Component/AnimNextComponent.h"
+#include "Module/AnimNextModuleInstance.h"
+
+void FAnimNextActorComponentReferenceComponent::OnInitializeHelper(UScriptStruct* InScriptStruct)
+{
+	check(InScriptStruct->IsChildOf(FAnimNextActorComponentReferenceComponent::StaticStruct()));
+
+	FAnimNextModuleInstance& ModuleInstance = GetModuleInstance();
+	UAnimNextComponent* AnimNextComponent = Cast<UAnimNextComponent>(ModuleInstance.GetObject());
+
+	FAnimNextModuleInstance::RunTaskOnGameThread(
+		[WeakAnimNextComponent = TWeakObjectPtr<UAnimNextComponent>(AnimNextComponent), ComponentType = ComponentType, InScriptStruct]()
+		{
+			UAnimNextComponent* Component = WeakAnimNextComponent.Get();
+			if(Component == nullptr)
+			{
+				return;
+			}
+
+			AActor* Owner = Component->GetOwner();
+			if(Owner == nullptr)
+			{
+				return;
+			}
+
+			UActorComponent* FoundComponent = Owner->FindComponentByClass(ComponentType);
+			Component->QueueTask(NAME_None, [FoundComponent, InScriptStruct](const UE::AnimNext::FModuleTaskContext& InContext)
+			{
+				InContext.TryAccessComponent(InScriptStruct, [FoundComponent](FAnimNextModuleInstanceComponent& InComponent)
+				{
+					static_cast<FAnimNextActorComponentReferenceComponent&>(InComponent).Component = FoundComponent;
+				});
+			});
+		});
+}
diff --git a/Engine/Plugins/Experimental/AnimNext/Source/AnimNext/Private/Module/AnimNextModuleInstance.cpp b/Engine/Plugins/Experimental/AnimNext/Source/AnimNext/Private/Module/AnimNextModuleInstance.cpp
--- a/Engine/Plugins/Experimental/AnimNext/Source/AnimNext/Private/Module/AnimNextModuleInstance.cpp
+++ b/Engine/Plugins/Experimental/AnimNext/Source/AnimNext/Private/Module/AnimNextModuleInstance.cpp
@@ -265,6 +265,14 @@ void FAnimNextModuleInstance::Initialize()
 	}
 }
 
+void FAnimNextModuleInstance::Uninitialize()
+{
+	for(TPair<FName, TInstancedStruct<FAnimNextModuleInstanceComponent>>& ComponentPair : ComponentMap)
+	{
+		ComponentPair.Value.GetMutable<FAnimNextModuleInstanceComponent>().Uninitialize();
+	}
+}
+
 void FAnimNextModuleInstance::ResetBindingsAndInstanceData()
 {
 	using namespace UE::AnimNext;
diff --git a/Engine/Plugins/Experimental/AnimNext/Source/AnimNext/Private/Module/AnimNextModuleInstanceComponent.cpp b/Engine/Plugins/Experimental/AnimNext/Source/AnimNext/Private/Module/AnimNextModuleInstanceComponent.cpp
new file mode 100644
--- /dev/null
+++ b/Engine/Plugins/Experimental/AnimNext/Source/AnimNext/Private/Module/AnimNextModuleInstanceComponent.cpp
@@ -0,0 +1,17 @@
+﻿// Copyright Epic Games, Inc. All Rights Reserved.
+
+#include "Module/AnimNextModuleInstanceComponent.h"
+
+void FAnimNextModuleInstanceComponent::Initialize(FAnimNextModuleInstance& InOwnerInstance)
+{
+	OwnerInstance = &InOwnerInstance;
+
+	OnInitialize();
+}
+
+void FAnimNextModuleInstanceComponent::Uninitialize()
+{
+	OnUninitialize();
+
+	OwnerInstance = nullptr;
+}
\ No newline at end of file
diff --git a/Engine/Plugins/Experimental/AnimNext/Source/AnimNext/Private/Module/AnimNextSkeletalMeshComponentReferenceComponent.cpp b/Engine/Plugins/Experimental/AnimNext/Source/AnimNext/Private/Module/AnimNextSkeletalMeshComponentReferenceComponent.cpp
new file mode 100644
--- /dev/null
+++ b/Engine/Plugins/Experimental/AnimNext/Source/AnimNext/Private/Module/AnimNextSkeletalMeshComponentReferenceComponent.cpp
@@ -0,0 +1,20 @@
+﻿// Copyright Epic Games, Inc. All Rights Reserved.
+
+#include "Module/AnimNextSkeletalMeshComponentReferenceComponent.h"
+
+#include "Components/SkeletalMeshComponent.h"
+
+FAnimNextSkeletalMeshComponentReferenceComponent::FAnimNextSkeletalMeshComponentReferenceComponent()
+{
+	ComponentType = USkeletalMeshComponent::StaticClass();
+}
+
+USkeletalMeshComponent* FAnimNextSkeletalMeshComponentReferenceComponent::GetComponent() const
+{
+	return Cast<USkeletalMeshComponent>(Component);
+}
+
+void FAnimNextSkeletalMeshComponentReferenceComponent::OnInitialize()
+{
+	OnInitializeHelper(StaticStruct());
+}
\ No newline at end of file
diff --git a/Engine/Plugins/Experimental/AnimNext/Source/AnimNext/Private/Module/ModuleInstanceComponent.cpp b/Engine/Plugins/Experimental/AnimNext/Source/AnimNext/Private/Module/ModuleInstanceComponent.cpp
deleted file mode 100644
diff --git a/Engine/Plugins/Experimental/AnimNext/Source/AnimNext/Private/Module/ModuleTaskContext.cpp b/Engine/Plugins/Experimental/AnimNext/Source/AnimNext/Private/Module/ModuleTaskContext.cpp
--- a/Engine/Plugins/Experimental/AnimNext/Source/AnimNext/Private/Module/ModuleTaskContext.cpp
+++ b/Engine/Plugins/Experimental/AnimNext/Source/AnimNext/Private/Module/ModuleTaskContext.cpp
@@ -20,4 +20,19 @@ void FModuleTaskContext::QueueInputTraitEvent(FAnimNextTraitEventPtr Event) cons
 	ModuleInstance->QueueInputTraitEvent(MoveTemp(Event));
 }
 
+void FModuleTaskContext::TryAccessComponent(UScriptStruct* InComponentType, TFunctionRef<void(FAnimNextModuleInstanceComponent&)> InFunction) const
+{
+	FModuleWriteGuard Guard(ModuleInstance);
+
+	FName ComponentName = InComponentType->GetFName();
+	uint32 ComponentNameHash = GetTypeHash(ComponentName);
+	FAnimNextModuleInstanceComponent* Component = ModuleInstance->TryGetComponent(ComponentNameHash, ComponentName);
+	if(Component == nullptr)
+	{
+		return;
+	}
+
+	InFunction(*Component);
+}
+
 }
diff --git a/Engine/Plugins/Experimental/AnimNext/Source/AnimNext/Public/Component/AnimNextComponent.h b/Engine/Plugins/Experimental/AnimNext/Source/AnimNext/Public/Component/AnimNextComponent.h
--- a/Engine/Plugins/Experimental/AnimNext/Source/AnimNext/Public/Component/AnimNextComponent.h
+++ b/Engine/Plugins/Experimental/AnimNext/Source/AnimNext/Public/Component/AnimNextComponent.h
@@ -9,6 +9,8 @@
 #include "TraitCore/TraitEvent.h"
 #include "Variables/IAnimNextVariableProxyHost.h"
 #include "Module/AnimNextModuleInitMethod.h"
+#include "Module/ModuleTaskContext.h"
+#include "Module/TaskRunLocation.h"
 
 #include "AnimNextComponent.generated.h"
 
@@ -193,6 +195,9 @@ public:
 	UFUNCTION(BlueprintCallable, Category = "AnimNext")
 	ANIMNEXT_API void ShowDebugDrawing(bool bShowDebugDrawing);
 
+	// Queue a task to run during execution 
+	ANIMNEXT_API void QueueTask(FName InModuleEventName, TUniqueFunction<void(const UE::AnimNext::FModuleTaskContext&)>&& InTaskFunction, UE::AnimNext::ETaskRunLocation InLocation = UE::AnimNext::ETaskRunLocation::Before);
+
 	// Queues an input trait event
 	// Input events will be processed in the next graph update after they are queued
 	ANIMNEXT_API void QueueInputTraitEvent(FAnimNextTraitEventPtr Event);
diff --git a/Engine/Plugins/Experimental/AnimNext/Source/AnimNext/Public/Module/ModuleTaskContext.h b/Engine/Plugins/Experimental/AnimNext/Source/AnimNext/Public/Module/ModuleTaskContext.h
new file mode 100644
--- /dev/null
+++ b/Engine/Plugins/Experimental/AnimNext/Source/AnimNext/Public/Module/ModuleTaskContext.h
@@ -0,0 +1,54 @@
+// Copyright Epic Games, Inc. All Rights Reserved.
+
+#pragma once
+
+#include "CoreMinimal.h"
+
+#include "TraitCore/TraitEvent.h"
+
+struct FAnimNextModuleInstanceComponent;
+struct FAnimNextModuleInstance;
+struct FInstancedPropertyBag;
+struct FAnimNextModuleInjectionComponent;
+
+namespace UE::AnimNext
+{
+	struct FScheduleContext;
+	enum class EParameterScopeOrdering : int32;
+	struct FModuleEventTickFunction;
+}
+
+namespace UE::AnimNext
+{
+
+// Context passed to schedule task callbacks
+struct FModuleTaskContext
+{
+public:
+	// Queues an input trait event
+	// Input events will be processed in the next graph update after they are queued
+	ANIMNEXT_API void QueueInputTraitEvent(FAnimNextTraitEventPtr Event) const;
+
+	// Access a module instance component of the specified type. If the component exists, then InFunction will be called
+	ANIMNEXT_API void TryAccessComponent(UScriptStruct* InComponentType, TFunctionRef<void(FAnimNextModuleInstanceComponent&)> InFunction) const;
+
+	template<typename ComponentType>
+	void TryAccessComponent(TFunctionRef<void(ComponentType&)> InFunction) const
+	{
+		TryAccessComponent(ComponentType::StaticStruct(), [&InFunction](FAnimNextModuleInstanceComponent& InComponent)
+		{
+			InFunction(static_cast<ComponentType&>(InComponent));
+		});
+	}
+
+private:
+	FModuleTaskContext(FAnimNextModuleInstance& InModuleInstance);
+
+	// The module instance currently running
+	FAnimNextModuleInstance* ModuleInstance;
+
+	friend UE::AnimNext::FModuleEventTickFunction;
+	friend ::FAnimNextModuleInjectionComponent;
+};
+
+}
\ No newline at end of file
diff --git a/Engine/Plugins/Experimental/AnimNext/Source/AnimNext/Public/Module/TaskRunLocation.h b/Engine/Plugins/Experimental/AnimNext/Source/AnimNext/Public/Module/TaskRunLocation.h
new file mode 100644
--- /dev/null
+++ b/Engine/Plugins/Experimental/AnimNext/Source/AnimNext/Public/Module/TaskRunLocation.h
@@ -0,0 +1,19 @@
+﻿// Copyright Epic Games, Inc. All Rights Reserved.
+
+#pragma once
+
+#include "CoreTypes.h"
+
+namespace UE::AnimNext
+{
+
+enum class ETaskRunLocation : int32
+{
+	// Run the task before the specified task
+	Before,
+
+	// Run the task after the specified task
+	After,
+};
+
+}
\ No newline at end of file
diff --git a/Engine/Plugins/Experimental/AnimNextAnimGraph/Source/AnimNextAnimGraph/Internal/Graph/AnimNextModuleAnimGraphComponent.h b/Engine/Plugins/Experimental/AnimNextAnimGraph/Source/AnimNextAnimGraph/Internal/Graph/AnimNextModuleAnimGraphComponent.h
--- a/Engine/Plugins/Experimental/AnimNextAnimGraph/Source/AnimNextAnimGraph/Internal/Graph/AnimNextModuleAnimGraphComponent.h
+++ b/Engine/Plugins/Experimental/AnimNextAnimGraph/Source/AnimNextAnimGraph/Internal/Graph/AnimNextModuleAnimGraphComponent.h
@@ -3,7 +3,7 @@
 #pragma once
 
 #include "Injection/InjectionInfo.h"
-#include "Module/ModuleInstanceComponent.h"
+#include "Module/AnimNextModuleInstanceComponent.h"
 #include "AnimNextModuleAnimGraphComponent.generated.h"
 
 struct FAnimNextGraphInstance;
diff --git a/Engine/Plugins/Experimental/AnimNextAnimGraph/Source/AnimNextAnimGraph/Private/Injection/AnimNextModuleInjectionComponent.cpp b/Engine/Plugins/Experimental/AnimNextAnimGraph/Source/AnimNextAnimGraph/Private/Injection/AnimNextModuleInjectionComponent.cpp
--- a/Engine/Plugins/Experimental/AnimNextAnimGraph/Source/AnimNextAnimGraph/Private/Injection/AnimNextModuleInjectionComponent.cpp
+++ b/Engine/Plugins/Experimental/AnimNextAnimGraph/Source/AnimNextAnimGraph/Private/Injection/AnimNextModuleInjectionComponent.cpp
@@ -10,12 +10,13 @@
 #include "Module/ModuleTaskContext.h"
 
 
-void FAnimNextModuleInjectionComponent::OnInitialize(FAnimNextModuleInstance& InOwnerInstance)
+void FAnimNextModuleInjectionComponent::OnInitialize()
 {
-	InjectionInfo = UE::AnimNext::FInjectionInfo(InOwnerInstance);
+	FAnimNextModuleInstance& ModuleInstance = GetModuleInstance();
+	InjectionInfo = UE::AnimNext::FInjectionInfo(ModuleInstance);
 	
 	// Register re-injection for each (user) tick function
-	for(UE::AnimNext::FModuleEventTickFunction& TickFunction : InOwnerInstance.GetTickFunctions())
+	for(UE::AnimNext::FModuleEventTickFunction& TickFunction : ModuleInstance.GetTickFunctions())
 	{
 		if(TickFunction.bUserEvent)
 		{
diff --git a/Engine/Plugins/Experimental/AnimNextAnimGraph/Source/AnimNextAnimGraph/Private/Injection/AnimNextModuleInjectionComponent.h b/Engine/Plugins/Experimental/AnimNextAnimGraph/Source/AnimNextAnimGraph/Private/Injection/AnimNextModuleInjectionComponent.h
--- a/Engine/Plugins/Experimental/AnimNextAnimGraph/Source/AnimNextAnimGraph/Private/Injection/AnimNextModuleInjectionComponent.h
+++ b/Engine/Plugins/Experimental/AnimNextAnimGraph/Source/AnimNextAnimGraph/Private/Injection/AnimNextModuleInjectionComponent.h
@@ -5,7 +5,7 @@
 #include "InjectionInfo.h"
 #include "DataInterface/AnimNextDataInterfaceHost.h"
 #include "Injection/InjectionRequest.h"
-#include "Module/ModuleInstanceComponent.h"
+#include "Module/AnimNextModuleInstanceComponent.h"
 #include "AnimNextModuleInjectionComponent.generated.h"
 
 namespace UE::AnimNext
@@ -28,7 +28,7 @@ struct FAnimNextModuleInjectionComponent : public FAnimNextModuleInstanceCompone
 
 private:
 	// FAnimNextModuleInstanceComponent interface
-	virtual void OnInitialize(FAnimNextModuleInstance& InOwnerInstance) override;
+	virtual void OnInitialize() override;
 	virtual void OnTraitEvent(FAnimNextTraitEvent& Event) override;
 
 	void OnInjectionEvent(UE::AnimNext::FInjection_InjectEvent& InEvent);
diff --git a/Engine/Plugins/Experimental/AnimNextAnimGraph/Source/AnimNextAnimGraph/Private/Traits/NotifyDispatcher.h b/Engine/Plugins/Experimental/AnimNextAnimGraph/Source/AnimNextAnimGraph/Private/Traits/NotifyDispatcher.h
--- a/Engine/Plugins/Experimental/AnimNextAnimGraph/Source/AnimNextAnimGraph/Private/Traits/NotifyDispatcher.h
+++ b/Engine/Plugins/Experimental/AnimNextAnimGraph/Source/AnimNextAnimGraph/Private/Traits/NotifyDispatcher.h
@@ -8,7 +8,7 @@
 #include "TraitInterfaces/ITimeline.h"
 #include "TraitInterfaces/IUpdate.h"
 #include "Animation/AnimNotifyQueue.h"
-#include "Module/ModuleInstanceComponent.h"
+#include "Module/AnimNextModuleInstanceComponent.h"
 
 #include "NotifyDispatcher.generated.h"
 
