AnimNext: CIS: Fix shadow variable error

diff --git a/Engine/Plugins/Experimental/AnimNext/Source/AnimNext/Private/Module/AnimNextActorComponentReferenceComponent.cpp b/Engine/Plugins/Experimental/AnimNext/Source/AnimNext/Private/Module/AnimNextActorComponentReferenceComponent.cpp
--- a/Engine/Plugins/Experimental/AnimNext/Source/AnimNext/Private/Module/AnimNextActorComponentReferenceComponent.cpp
+++ b/Engine/Plugins/Experimental/AnimNext/Source/AnimNext/Private/Module/AnimNextActorComponentReferenceComponent.cpp
@@ -1,4 +1,4 @@
-﻿// Copyright Epic Games, Inc. All Rights Reserved.
+// Copyright Epic Games, Inc. All Rights Reserved.
 
 #include "Module/AnimNextActorComponentReferenceComponent.h"
 
@@ -15,20 +15,20 @@ void FAnimNextActorComponentReferenceComponent::OnInitializeHelper(UScriptStruct
 	FAnimNextModuleInstance::RunTaskOnGameThread(
 		[WeakAnimNextComponent = TWeakObjectPtr<UAnimNextComponent>(AnimNextComponent), ComponentType = ComponentType, InScriptStruct]()
 		{
-			UAnimNextComponent* Component = WeakAnimNextComponent.Get();
-			if(Component == nullptr)
+			UAnimNextComponent* PinnedAnimNextComponent = WeakAnimNextComponent.Get();
+			if(PinnedAnimNextComponent == nullptr)
 			{
 				return;
 			}
 
-			AActor* Owner = Component->GetOwner();
+			AActor* Owner = PinnedAnimNextComponent->GetOwner();
 			if(Owner == nullptr)
 			{
 				return;
 			}
 
 			UActorComponent* FoundComponent = Owner->FindComponentByClass(ComponentType);
-			Component->QueueTask(NAME_None, [FoundComponent, InScriptStruct](const UE::AnimNext::FModuleTaskContext& InContext)
+			PinnedAnimNextComponent->QueueTask(NAME_None, [FoundComponent, InScriptStruct](const UE::AnimNext::FModuleTaskContext& InContext)
 			{
 				InContext.TryAccessComponent(InScriptStruct, [FoundComponent](FAnimNextModuleInstanceComponent& InComponent)
 				{
