AnimNext World Subsystem : Fixed shutdown issues at Traits where the instance data destruct was called too late and had objects already destroyed, by forcing module instances to be destoyed at BeginDestroy of the World Susbsystem, rather than destructor

diff --git a/Engine/Plugins/Experimental/AnimNext/Source/AnimNext/Internal/Component/AnimNextWorldSubsystem.h b/Engine/Plugins/Experimental/AnimNext/Source/AnimNext/Internal/Component/AnimNextWorldSubsystem.h
--- a/Engine/Plugins/Experimental/AnimNext/Source/AnimNext/Internal/Component/AnimNextWorldSubsystem.h
+++ b/Engine/Plugins/Experimental/AnimNext/Source/AnimNext/Internal/Component/AnimNextWorldSubsystem.h
@@ -74,7 +74,8 @@ class ANIMNEXT_API UAnimNextWorldSubsystem : public UWorldSubsystem
 
 public:
 	UAnimNextWorldSubsystem();
-	virtual ~UAnimNextWorldSubsystem() override;
+
+	virtual void BeginDestroy() override;
 
 	// Get the subsystem for the specified object's world
 	static UAnimNextWorldSubsystem* Get(UObject* InObject);
diff --git a/Engine/Plugins/Experimental/AnimNext/Source/AnimNext/Private/Component/AnimNextWorldSubsystem.cpp b/Engine/Plugins/Experimental/AnimNext/Source/AnimNext/Private/Component/AnimNextWorldSubsystem.cpp
--- a/Engine/Plugins/Experimental/AnimNext/Source/AnimNext/Private/Component/AnimNextWorldSubsystem.cpp
+++ b/Engine/Plugins/Experimental/AnimNext/Source/AnimNext/Private/Component/AnimNextWorldSubsystem.cpp
@@ -31,15 +31,26 @@ UAnimNextWorldSubsystem::UAnimNextWorldSubsystem()
 	}
 }
 
-UAnimNextWorldSubsystem::~UAnimNextWorldSubsystem()
+void UAnimNextWorldSubsystem::BeginDestroy()
 {
-	if(!HasAnyFlags(RF_ClassDefaultObject))
+	Super::BeginDestroy();
+
+	if (!HasAnyFlags(RF_ClassDefaultObject))
 	{
 #if WITH_EDITOR
 		UAnimNextModule::OnModuleCompiled().Remove(OnModuleCompiledHandle);
 #endif
 
 		FWorldDelegates::OnWorldPreActorTick.Remove(OnWorldPreActorTickHandle);
+
+		{
+			FRWScopeLock InstancesLockScope(InstancesLock, SLT_Write);
+			for (FAnimNextModuleInstance& Instance : Instances)
+			{
+				Instance.Uninitialize();
+			}
+			Instances = UE::AnimNext::TPool<FAnimNextModuleInstance>(); // Force instances destruction
+		}
 	}
 }
 
