Mark the RigUnit task method with a parallel game thread task tag to fix crash in skeletal mesh updater which calls check(IsInParallelGameThread());

diff --git a/Engine/Plugins/Experimental/AnimNext/Source/AnimNext/Private/Graph/RigUnit_AnimNextWriteSkeletalMeshComponentPose.cpp b/Engine/Plugins/Experimental/AnimNext/Source/AnimNext/Private/Graph/RigUnit_AnimNextWriteSkeletalMeshComponentPose.cpp
--- a/Engine/Plugins/Experimental/AnimNext/Source/AnimNext/Private/Graph/RigUnit_AnimNextWriteSkeletalMeshComponentPose.cpp
+++ b/Engine/Plugins/Experimental/AnimNext/Source/AnimNext/Private/Graph/RigUnit_AnimNextWriteSkeletalMeshComponentPose.cpp
@@ -18,6 +18,7 @@ FRigUnit_AnimNextWriteSkeletalMeshComponentPose_Execute()
 	SCOPE_CYCLE_COUNTER(STAT_AnimNext_Write_Pose);
 
 	using namespace UE::AnimNext;
+	FTaskTagScope TaskTag(ETaskTag::EParallelGameThread);
 
 	USkeletalMeshComponent* OutputComponent = SkeletalMeshComponent;
 
