Add a check to ensure we use the same AnimNext reference pose when writing out to a skeletal mesh component

diff --git a/Engine/Plugins/Experimental/AnimNext/Source/AnimNext/Private/Graph/RigUnit_AnimNextWriteSkeletalMeshComponentPose.cpp b/Engine/Plugins/Experimental/AnimNext/Source/AnimNext/Private/Graph/RigUnit_AnimNextWriteSkeletalMeshComponentPose.cpp
--- a/Engine/Plugins/Experimental/AnimNext/Source/AnimNext/Private/Graph/RigUnit_AnimNextWriteSkeletalMeshComponentPose.cpp
+++ b/Engine/Plugins/Experimental/AnimNext/Source/AnimNext/Private/Graph/RigUnit_AnimNextWriteSkeletalMeshComponentPose.cpp
@@ -50,13 +50,21 @@ FRigUnit_AnimNextWriteSkeletalMeshComponentPose_Execute()
 	// We cannot write to the skeletal mesh component if it is using anim BP
 	if(OutputComponent->bEnableAnimation)
 	{
-		UE_LOG(LogAnimation, Warning, TEXT("AnimNext: Could not write to skeletal mesh component - bEnableAnimation is true"))
+		UE_LOG(LogAnimation, Warning, TEXT("AnimNext: Could not write to skeletal mesh component - bEnableAnimation is true"));
 		return;
 	}
 
 	FDataHandle RefPoseHandle = FDataRegistry::Get()->GetOrGenerateReferencePose(OutputComponent);
 	const UE::AnimNext::FReferencePose& RefPose = RefPoseHandle.GetRef<UE::AnimNext::FReferencePose>();
 
+	// We cannot use two different reference poses because we want to avoid an expensive remap operation
+	// You should remap the pose explicitly if this is what you need
+	if (&RefPose != &Pose.LODPose.GetRefPose())
+	{
+		UE_LOG(LogAnimation, Warning, TEXT("AnimNext: Could not write to skeletal mesh component - The input pose and the skeletal mesh component use different reference poses"));
+		return;
+	}
+
 	FMemMark MemMark(FMemStack::Get());
 
 	// TODO: The LOD level we got here is not guaranteed to be stable. We should find another way to get/bind the skeletal mesh component's LOD level.
