Anim Next rewind debugger integration

diff --git a/Engine/Plugins/Experimental/AnimNext/AnimNext.uplugin b/Engine/Plugins/Experimental/AnimNext/AnimNext.uplugin
--- a/Engine/Plugins/Experimental/AnimNext/AnimNext.uplugin
+++ b/Engine/Plugins/Experimental/AnimNext/AnimNext.uplugin
@@ -45,6 +45,10 @@
 		{
 			"Name": "HierarchyTableAnimation",
 			"Enabled": true
+		},
+		{
+			"Name": "GameplayInsights",
+			"Enabled": true
 		}
 	],
 	"Modules" :
diff --git a/Engine/Plugins/Experimental/AnimNext/Source/AnimNext/AnimNext.Build.cs b/Engine/Plugins/Experimental/AnimNext/Source/AnimNext/AnimNext.Build.cs
--- a/Engine/Plugins/Experimental/AnimNext/Source/AnimNext/AnimNext.Build.cs
+++ b/Engine/Plugins/Experimental/AnimNext/Source/AnimNext/AnimNext.Build.cs
@@ -14,7 +14,8 @@ namespace UnrealBuildTool.Rules
 					"RigVM",
 					"Engine",
 					"RenderCore",
-					"HierarchyTableRuntime"
+					"HierarchyTableRuntime",
+					"RewindDebuggerRuntimeInterface"
 				}
 			);
 
@@ -35,6 +36,7 @@ namespace UnrealBuildTool.Rules
 						"ToolMenus",
 						"Slate",
 						"SlateCore",
+						"TraceLog",
 						"UnrealEd"
 					}
 				);
diff --git a/Engine/Plugins/Experimental/AnimNext/Source/AnimNext/Internal/Module/AnimNextModuleInstance.h b/Engine/Plugins/Experimental/AnimNext/Source/AnimNext/Internal/Module/AnimNextModuleInstance.h
--- a/Engine/Plugins/Experimental/AnimNext/Source/AnimNext/Internal/Module/AnimNextModuleInstance.h
+++ b/Engine/Plugins/Experimental/AnimNext/Source/AnimNext/Internal/Module/AnimNextModuleInstance.h
@@ -147,6 +147,7 @@ struct FAnimNextModuleInstance : public FAnimNextDataInterfaceInstance
 	// Get tick functions for this module instance
 	ANIMNEXT_API TArrayView<UE::AnimNext::FModuleEventTickFunction> GetTickFunctions();
 
+	ANIMNEXT_API UE::AnimNext::FModuleHandle GetHandle() const { return Handle; }
 private:
 	// Setup the entry
 	void Initialize();
diff --git a/Engine/Plugins/Experimental/AnimNext/Source/AnimNext/Private/AnimNextModuleImpl.cpp b/Engine/Plugins/Experimental/AnimNext/Source/AnimNext/Private/AnimNextModuleImpl.cpp
--- a/Engine/Plugins/Experimental/AnimNext/Source/AnimNext/Private/AnimNextModuleImpl.cpp
+++ b/Engine/Plugins/Experimental/AnimNext/Source/AnimNext/Private/AnimNextModuleImpl.cpp
@@ -24,11 +24,18 @@
 #include "Variables/AnimNextSoftFunctionPtr.h"
 #include "RigVMCore/RigVMGraphFunctionDefinition.h"
 #include "HierarchyTable.h"
+#include "Features/IModularFeatures.h"
+#include "RewindDebuggerRuntimeInterface/IRewindDebuggerRuntimeExtension.h"
+#include "RewindDebugger/RewindDebuggerAnimNextRuntime.h"
 
 #define LOCTEXT_NAMESPACE "AnimNextModule"
 
 namespace UE::AnimNext
 {
+#if ANIMNEXT_TRACE_ENABLED
+	FRewindDebuggerAnimNextRuntime GRewindDebuggerAnimNextRuntime;
+#endif
+	
 	void FAnimNextModuleImpl::StartupModule()
 	{
 		GetMutableDefault<UAnimNextConfig>()->LoadConfig();
@@ -67,6 +74,10 @@ namespace UE::AnimNext
 
 		FRigVMRuntimeDataRegistry::Init();
 
+#if ANIMNEXT_TRACE_ENABLED
+		IModularFeatures::Get().RegisterModularFeature(IRewindDebuggerRuntimeExtension::ModularFeatureName, &GRewindDebuggerAnimNextRuntime);
+#endif
+
 		UE::UniversalObjectLocator::IUniversalObjectLocatorModule& UolModule = FModuleManager::Get().LoadModuleChecked<UE::UniversalObjectLocator::IUniversalObjectLocatorModule>("UniversalObjectLocator");
 		FDelayedAutoRegisterHelper(EDelayedRegisterRunPhase::ObjectSystemReady,
 			[&UolModule]
@@ -102,6 +113,10 @@ namespace UE::AnimNext
 
 	void FAnimNextModuleImpl::ShutdownModule()
 	{
+#if ANIMNEXT_TRACE_ENABLED
+		IModularFeatures::Get().UnregisterModularFeature(IRewindDebuggerRuntimeExtension::ModularFeatureName, &GRewindDebuggerAnimNextRuntime);
+#endif
+		
 		FRigVMRuntimeDataRegistry::Destroy();
 		FDataRegistry::Destroy();
 	}
diff --git a/Engine/Plugins/Experimental/AnimNext/Source/AnimNext/Private/Module/AnimNextModuleInstance.cpp b/Engine/Plugins/Experimental/AnimNext/Source/AnimNext/Private/Module/AnimNextModuleInstance.cpp
--- a/Engine/Plugins/Experimental/AnimNext/Source/AnimNext/Private/Module/AnimNextModuleInstance.cpp
+++ b/Engine/Plugins/Experimental/AnimNext/Source/AnimNext/Private/Module/AnimNextModuleInstance.cpp
@@ -18,6 +18,7 @@
 #include "Module/ProxyVariablesContext.h"
 #include "UObject/UObjectIterator.h"
 #include "Variables/IAnimNextVariableProxyHost.h"
+#include "RewindDebugger/AnimNextTrace.h"
 
 DEFINE_STAT(STAT_AnimNext_InitializeInstance);
 
@@ -426,6 +427,8 @@ void FAnimNextModuleInstance::CopyProxyVariables()
 			PublicVariablesProxy.bIsDirty = false;
 		}
 	}
+	
+	TRACE_ANIMNEXT_MODULE_VARIABLES(this);
 }
 
 const UAnimNextModule* FAnimNextModuleInstance::GetModule() const
diff --git a/Engine/Plugins/Experimental/AnimNext/Source/AnimNext/Private/RewindDebugger/AnimNextTrace.cpp b/Engine/Plugins/Experimental/AnimNext/Source/AnimNext/Private/RewindDebugger/AnimNextTrace.cpp
new file mode 100644
--- /dev/null
+++ b/Engine/Plugins/Experimental/AnimNext/Source/AnimNext/Private/RewindDebugger/AnimNextTrace.cpp
@@ -0,0 +1,96 @@
+// Copyright Epic Games, Inc. All Rights Reserved.
+
+#include "AnimNextTrace.h"
+
+#include "ObjectTrace.h"
+#include "Module/AnimNextModuleInstance.h"
+#include "Serialization/ObjectWriter.h"
+
+#if ANIMNEXT_TRACE_ENABLED
+
+
+UE_TRACE_EVENT_BEGIN(AnimNext, Module)
+	UE_TRACE_EVENT_FIELD(uint64, ModuleInstanceId)
+	UE_TRACE_EVENT_FIELD(uint64, OuterObjectId)
+	UE_TRACE_EVENT_FIELD(uint64, ModuleAssetId)
+UE_TRACE_EVENT_END()
+
+UE_TRACE_EVENT_BEGIN(AnimNext, ModuleVariables)
+	UE_TRACE_EVENT_FIELD(uint64, Cycle)
+	UE_TRACE_EVENT_FIELD(double, RecordingTime)
+	UE_TRACE_EVENT_FIELD(uint64, ModuleInstanceId)
+	UE_TRACE_EVENT_FIELD(uint8[], VariableData)
+UE_TRACE_EVENT_END()
+
+UE_TRACE_CHANNEL_DEFINE(AnimNextChannel);
+
+namespace UE::AnimNext
+{
+	
+TSet<uint64> GTracedModules;
+
+void FAnimNextTrace::Reset()
+{
+	GTracedModules.Empty();
+}
+
+void FAnimNextTrace::OutputAnimNextModuleInstance(const FAnimNextModuleInstance* ModuleInstance)
+{
+	bool bChannelEnabled = UE_TRACE_CHANNELEXPR_IS_ENABLED(AnimNextChannel);
+	if (!bChannelEnabled || ModuleInstance == nullptr)
+	{
+		return;
+	}
+
+	if (CANNOT_TRACE_OBJECT(ModuleInstance->GetObject()))
+	{
+		return;
+	}
+	uint64 ModuleId = ModuleInstance->GetHandle().GetUniqueId();
+
+	if (!GTracedModules.Contains(ModuleId))
+	{
+		GTracedModules.Add(ModuleId);
+
+		TRACE_OBJECT(ModuleInstance->GetObject());
+		TRACE_OBJECT(ModuleInstance->GetModule());
+
+		UE_TRACE_LOG(AnimNext, Module, AnimNextChannel)
+			<< Module.ModuleInstanceId(ModuleId)
+			<< Module.OuterObjectId(FObjectTrace::GetObjectId(ModuleInstance->GetObject()))
+			<< Module.ModuleAssetId(FObjectTrace::GetObjectId(ModuleInstance->GetModule()));
+	}
+}
+
+void FAnimNextTrace::OutputAnimNextModuleInstanceVariables(const FAnimNextModuleInstance* ModuleInstance)
+{
+	bool bChannelEnabled = UE_TRACE_CHANNELEXPR_IS_ENABLED(AnimNextChannel);
+	if (!bChannelEnabled || ModuleInstance == nullptr)
+	{
+		return;
+	}
+
+	if (CANNOT_TRACE_OBJECT(ModuleInstance->GetObject()))
+	{
+		return;
+	}
+	uint64 ModuleId = ModuleInstance->GetHandle().GetUniqueId();
+
+	TRACE_ANIMNEXT_MODULE(ModuleInstance);
+
+	const FInstancedPropertyBag Variables = ModuleInstance->GetVariables();
+
+	TArray<uint8> ArchiveData;
+	FObjectWriter Archive(ArchiveData);
+	const_cast<FInstancedPropertyBag&>(Variables).Serialize(Archive);
+
+	UE_TRACE_LOG(AnimNext, ModuleVariables, AnimNextChannel)
+			<< ModuleVariables.Cycle(FPlatformTime::Cycles64())
+			<< ModuleVariables.RecordingTime(FObjectTrace::GetWorldElapsedTime(ModuleInstance->GetObject()->GetWorld()))
+			<< ModuleVariables.ModuleInstanceId(ModuleId)
+			<< ModuleVariables.VariableData(ArchiveData.GetData(), ArchiveData.Num());
+}
+
+}
+
+#endif
diff --git a/Engine/Plugins/Experimental/AnimNext/Source/AnimNext/Private/RewindDebugger/AnimNextTrace.h b/Engine/Plugins/Experimental/AnimNext/Source/AnimNext/Private/RewindDebugger/AnimNextTrace.h
new file mode 100644
--- /dev/null
+++ b/Engine/Plugins/Experimental/AnimNext/Source/AnimNext/Private/RewindDebugger/AnimNextTrace.h
@@ -0,0 +1,36 @@
+// Copyright Epic Games, Inc. All Rights Reserved.
+
+#pragma once
+
+#include "CoreTypes.h"
+#include "ObjectTrace.h"
+#include "CoreMinimal.h"
+#include "Trace/Trace.h"
+
+#define ANIMNEXT_TRACE_ENABLED (OBJECT_TRACE_ENABLED && !(UE_BUILD_SHIPPING || UE_BUILD_TEST))
+
+#if ANIMNEXT_TRACE_ENABLED
+
+UE_TRACE_CHANNEL_EXTERN(AnimNextChannel, ANIMNEXT_API);
+
+struct FAnimNextModuleInstance;
+
+namespace UE::AnimNext
+{
+
+struct FAnimNextTrace
+{
+	ANIMNEXT_API static void Reset();
+	ANIMNEXT_API static void OutputAnimNextModuleInstance(const FAnimNextModuleInstance* ModuleInstance);
+	ANIMNEXT_API static void OutputAnimNextModuleInstanceVariables(const FAnimNextModuleInstance* ModuleInstance);
+};
+
+}
+
+
+#define TRACE_ANIMNEXT_MODULE(Module) UE::AnimNext::FAnimNextTrace::OutputAnimNextModuleInstance(Module);
+#define TRACE_ANIMNEXT_MODULE_VARIABLES(Module) UE::AnimNext::FAnimNextTrace::OutputAnimNextModuleInstanceVariables(Module);
+#else
+#define TRACE_ANIMNEXT_MODULE(Module)
+#define TRACE_ANIMNEXT_MODULE_VARIABLES(Module)
+#endif
diff --git a/Engine/Plugins/Experimental/AnimNext/Source/AnimNext/Private/RewindDebugger/RewindDebuggerAnimNextRuntime.cpp b/Engine/Plugins/Experimental/AnimNext/Source/AnimNext/Private/RewindDebugger/RewindDebuggerAnimNextRuntime.cpp
new file mode 100644
--- /dev/null
+++ b/Engine/Plugins/Experimental/AnimNext/Source/AnimNext/Private/RewindDebugger/RewindDebuggerAnimNextRuntime.cpp
@@ -0,0 +1,28 @@
+// Copyright Epic Games, Inc. All Rights Reserved.
+
+#include "RewindDebuggerAnimNextRuntime.h"
+
+#if ANIMNEXT_TRACE_ENABLED
+
+namespace UE::AnimNext
+{
+	
+void FRewindDebuggerAnimNextRuntime::RecordingStarted()
+{
+	UE::Trace::ToggleChannel(TEXT("AnimNext"), true);
+}
+
+void FRewindDebuggerAnimNextRuntime::RecordingStopped()
+{
+	UE::Trace::ToggleChannel(TEXT("AnimNext"), false);
+}
+
+void FRewindDebuggerAnimNextRuntime::Clear()
+{
+	FAnimNextTrace::Reset();
+}
+
+}
+
+#endif
+	
diff --git a/Engine/Plugins/Experimental/AnimNext/Source/AnimNext/Private/RewindDebugger/RewindDebuggerAnimNextRuntime.h b/Engine/Plugins/Experimental/AnimNext/Source/AnimNext/Private/RewindDebugger/RewindDebuggerAnimNextRuntime.h
new file mode 100644
--- /dev/null
+++ b/Engine/Plugins/Experimental/AnimNext/Source/AnimNext/Private/RewindDebugger/RewindDebuggerAnimNextRuntime.h
@@ -0,0 +1,25 @@
+// Copyright Epic Games, Inc. All Rights Reserved.
+#pragma once
+
+#include "CoreMinimal.h"
+#include "RewindDebuggerRuntimeInterface/IRewindDebuggerRuntimeExtension.h"
+#include "RewindDebugger/AnimNextTrace.h"
+
+#if ANIMNEXT_TRACE_ENABLED
+// Rewind debugger extension for Chooser support
+
+namespace UE::AnimNext
+{
+
+class FRewindDebuggerAnimNextRuntime : public IRewindDebuggerRuntimeExtension
+{
+public:
+	virtual void RecordingStarted() override;
+	virtual void RecordingStopped() override;
+	
+	virtual void Clear() override;
+};
+
+}
+
+#endif
\ No newline at end of file
diff --git a/Engine/Plugins/Experimental/AnimNext/Source/AnimNext/Public/AnimNextPoolHandle.h b/Engine/Plugins/Experimental/AnimNext/Source/AnimNext/Public/AnimNextPoolHandle.h
--- a/Engine/Plugins/Experimental/AnimNext/Source/AnimNext/Public/AnimNextPoolHandle.h
+++ b/Engine/Plugins/Experimental/AnimNext/Source/AnimNext/Public/AnimNextPoolHandle.h
@@ -26,6 +26,7 @@ public:
 		Index = MAX_uint32;
 		SerialNumber = 0;
 	}
+	uint64 GetUniqueId() const { return (static_cast<uint64>(Index) << 32) | static_cast<uint64>(SerialNumber);}
 
 	bool IsValid() const
 	{
diff --git a/Engine/Plugins/Experimental/AnimNext/Source/AnimNextEditor/AnimNextEditor.Build.cs b/Engine/Plugins/Experimental/AnimNext/Source/AnimNextEditor/AnimNextEditor.Build.cs
--- a/Engine/Plugins/Experimental/AnimNext/Source/AnimNextEditor/AnimNextEditor.Build.cs
+++ b/Engine/Plugins/Experimental/AnimNext/Source/AnimNextEditor/AnimNextEditor.Build.cs
@@ -49,6 +49,12 @@ namespace UnrealBuildTool.Rules
 					"EditorFramework",
 					"InteractiveToolsFramework",
 					"EditorInteractiveToolsFramework",
+					"TraceAnalysis",
+					"TraceLog",
+					"TraceServices",
+					"TraceInsights",
+					"RewindDebuggerInterface",
+					"GameplayInsights",
 				}
 			);
 		}
diff --git a/Engine/Plugins/Experimental/AnimNext/Source/AnimNextEditor/Private/AnimNextEditorModule.cpp b/Engine/Plugins/Experimental/AnimNext/Source/AnimNextEditor/Private/AnimNextEditorModule.cpp
--- a/Engine/Plugins/Experimental/AnimNext/Source/AnimNextEditor/Private/AnimNextEditorModule.cpp
+++ b/Engine/Plugins/Experimental/AnimNext/Source/AnimNextEditor/Private/AnimNextEditorModule.cpp
@@ -67,10 +67,23 @@
 #include "Common/GraphEditorSchemaActions.h"
 #include "EditorUtils.h"
 #include "PersonaModule.h"
+#include "Developer/TraceServices/Public/TraceServices/ModuleService.h"
+#include "Editor/RewindDebuggerInterface/Public/IRewindDebuggerExtension.h"
+#include "Editor/RewindDebuggerInterface/Public/IRewindDebuggerTrackCreator.h"
 #include "Widgets/Docking/SDockTab.h"
+#include "RewindDebugger/RewindDebuggerAnimNext.h"
+#include "RewindDebugger/AnimNextModuleTrack.h"
+#include "RewindDebugger/AnimNextTraceModule.h"
 
 #define LOCTEXT_NAMESPACE "AnimNextEditorModule"
 
+namespace
+{
+	FRewindDebuggerAnimNext GRewindDebuggerAnimNext;
+	UE::AnimNextEditor::FAnimNextModuleTrackCreator GAnimNextModulesTrackCreator;
+	FAnimNextTraceModule GAnimNextTraceModule;
+}
+
 namespace UE::AnimNext::Editor
 {
 
@@ -87,6 +100,10 @@ void FAnimNextEditorModule::StartupModule()
 		LOCTEXT("SettingsDescription", "Customize AnimNext Settings."),
 		GetMutableDefault<UAnimNextConfig>()
 	);
+	
+	IModularFeatures::Get().RegisterModularFeature(IRewindDebuggerExtension::ModularFeatureName, &GRewindDebuggerAnimNext);
+	IModularFeatures::Get().RegisterModularFeature(RewindDebugger::IRewindDebuggerTrackCreator::ModularFeatureName, &GAnimNextModulesTrackCreator);
+	IModularFeatures::Get().RegisterModularFeature(TraceServices::ModuleFeatureName, &GAnimNextTraceModule);
 
 	FPropertyEditorModule& PropertyModule = FModuleManager::LoadModuleChecked<FPropertyEditorModule>("PropertyEditor");
 
@@ -392,6 +409,10 @@ void FAnimNextEditorModule::ShutdownModule()
 		PersonaModule.UnregisterNotifyHostAsset(FTopLevelAssetPath(TEXT("/Script/AnimNext.AnimNextRigVMAsset")));
 	}
 	UnregisterAssetCompilationHandler(FTopLevelAssetPath(TEXT("/Script/AnimNext.AnimNextModule")));
+	
+	IModularFeatures::Get().UnregisterModularFeature(IRewindDebuggerExtension::ModularFeatureName, &GRewindDebuggerAnimNext);
+	IModularFeatures::Get().UnregisterModularFeature(RewindDebugger::IRewindDebuggerTrackCreator::ModularFeatureName, &GAnimNextModulesTrackCreator);
+	IModularFeatures::Get().UnregisterModularFeature(TraceServices::ModuleFeatureName, &GAnimNextTraceModule);
 }
 
 void FAnimNextEditorModule::RegisterLocatorFragmentEditorType(FName InLocatorFragmentEditorName)
diff --git a/Engine/Plugins/Experimental/AnimNext/Source/AnimNextEditor/Private/RewindDebugger/AnimNextAnalyzer.cpp b/Engine/Plugins/Experimental/AnimNext/Source/AnimNextEditor/Private/RewindDebugger/AnimNextAnalyzer.cpp
new file mode 100644
--- /dev/null
+++ b/Engine/Plugins/Experimental/AnimNext/Source/AnimNextEditor/Private/RewindDebugger/AnimNextAnalyzer.cpp
@@ -0,0 +1,58 @@
+// Copyright Epic Games, Inc. All Rights Reserved.
+
+#include "AnimNextAnalyzer.h"
+
+#include "HAL/LowLevelMemTracker.h"
+#include "Serialization/MemoryReader.h"
+#include "TraceServices/Model/AnalysisSession.h"
+#include "AnimNextProvider.h"
+#include "TraceServices/Utils.h"
+#include "StructUtils/PropertyBag.h"
+#include "Serialization/ObjectReader.h"
+
+FAnimNextAnalyzer::FAnimNextAnalyzer(TraceServices::IAnalysisSession& InSession, FAnimNextProvider& InProvider)
+	: Session(InSession)
+	, Provider(InProvider)
+{
+}
+
+void FAnimNextAnalyzer::OnAnalysisBegin(const FOnAnalysisContext& Context)
+{
+	auto& Builder = Context.InterfaceBuilder;
+
+	Builder.RouteEvent(RouteId_Module, "AnimNext", "Module");
+	Builder.RouteEvent(RouteId_ModuleVariables, "AnimNext", "ModuleVariables");
+}
+
+bool FAnimNextAnalyzer::OnEvent(uint16 RouteId, EStyle Style, const FOnEventContext& Context)
+{
+	LLM_SCOPE_BYNAME(TEXT("Insights/FAnimNextAnalyzer"));
+
+	TraceServices::FAnalysisSessionEditScope _(Session);
+
+	const auto& EventData = Context.EventData;
+	switch (RouteId)
+	{
+		case RouteId_Module:
+		{
+			uint64 ModuleInstanceId = EventData.GetValue<uint64>("ModuleInstanceId");
+			uint64 ModuleAssetId = EventData.GetValue<uint64>("ModuleAssetId");
+			uint64 OuterObjectId = EventData.GetValue<uint64>("OuterObjectId");
+			Provider.AppendModule(ModuleInstanceId, ModuleAssetId, OuterObjectId);
+			break;
+		}
+		case RouteId_ModuleVariables:
+		{
+			uint64 ModuleInstanceId = EventData.GetValue<uint64>("ModuleInstanceId");
+			uint64 Cycle = EventData.GetValue<uint64>("Cycle");
+			double RecordingTime = EventData.GetValue<double>("RecordingTime");
+
+			TArrayView<const uint8> VariableData = EventData.GetArrayView<uint8>("VariableData");
+
+			Provider.AppendVariables(Context.EventTime.AsSeconds(Cycle), RecordingTime, ModuleInstanceId, VariableData);
+			break;
+		}
+	}
+
+	return true;
+}
diff --git a/Engine/Plugins/Experimental/AnimNext/Source/AnimNextEditor/Private/RewindDebugger/AnimNextAnalyzer.h b/Engine/Plugins/Experimental/AnimNext/Source/AnimNextEditor/Private/RewindDebugger/AnimNextAnalyzer.h
new file mode 100644
--- /dev/null
+++ b/Engine/Plugins/Experimental/AnimNext/Source/AnimNextEditor/Private/RewindDebugger/AnimNextAnalyzer.h
@@ -0,0 +1,28 @@
+// Copyright Epic Games, Inc. All Rights Reserved.
+
+#pragma once
+
+#include "Trace/Analyzer.h"
+
+class FAnimNextProvider;
+namespace TraceServices { class IAnalysisSession; }
+
+class FAnimNextAnalyzer : public UE::Trace::IAnalyzer
+{
+public:
+	FAnimNextAnalyzer(TraceServices::IAnalysisSession& InSession, FAnimNextProvider& InProvider);
+
+	virtual void OnAnalysisBegin(const FOnAnalysisContext& Context) override;
+	virtual void OnAnalysisEnd() override {}
+	virtual bool OnEvent(uint16 RouteId, EStyle Style, const FOnEventContext& Context) override;
+
+private:
+	enum : uint16
+	{
+		RouteId_Module,
+		RouteId_ModuleVariables,
+	};
+
+	TraceServices::IAnalysisSession& Session;
+	FAnimNextProvider& Provider;
+};
diff --git a/Engine/Plugins/Experimental/AnimNext/Source/AnimNextEditor/Private/RewindDebugger/AnimNextModuleTrack.cpp b/Engine/Plugins/Experimental/AnimNext/Source/AnimNextEditor/Private/RewindDebugger/AnimNextModuleTrack.cpp
new file mode 100644
--- /dev/null
+++ b/Engine/Plugins/Experimental/AnimNext/Source/AnimNextEditor/Private/RewindDebugger/AnimNextModuleTrack.cpp
@@ -0,0 +1,193 @@
+// Copyright Epic Games, Inc. All Rights Reserved.
+
+#include "AnimNextModuleTrack.h"
+
+#include "AnimNextProvider.h"
+#include "Editor.h"
+#include "IGameplayProvider.h"
+#include "IRewindDebugger.h"
+#include "Subsystems/AssetEditorSubsystem.h"
+#include "Modules/ModuleManager.h"
+#include "PropertyEditorModule.h"
+#include "Serialization/ObjectReader.h"
+
+
+#define LOCTEXT_NAMESPACE "AnimNextModuleTrack"
+
+namespace UE::AnimNextEditor
+{
+
+static const FName AnimNextModulesName("AnimNextModules");
+
+ FName FAnimNextModuleTrackCreator::GetTargetTypeNameInternal() const
+ {
+ 	static const FName ObjectName("AnimNextComponent");
+	return ObjectName;
+}
+
+	
+FText FAnimNextModuleTrack::GetDisplayNameInternal() const
+{
+	TraceServices::FAnalysisSessionReadScope SessionReadScope(*IRewindDebugger::Instance()->GetAnalysisSession());
+	const FAnimNextProvider* AnimNextProvider = IRewindDebugger::Instance()->GetAnalysisSession()->ReadProvider<FAnimNextProvider>("AnimNextProvider");
+	const IGameplayProvider* GameplayProvider = IRewindDebugger::Instance()->GetAnalysisSession()->ReadProvider<IGameplayProvider>("GameplayProvider");
+
+ 	uint64 ModuleId;
+ 	if (AnimNextProvider->GetModuleId(ObjectId, ModuleId))
+ 	{
+ 		if (const FDataInterfaceData* Data = AnimNextProvider->GetDataInterfaceData(ModuleId))
+ 		{
+			const FObjectInfo& ModuleInfo = GameplayProvider->GetObjectInfo(Data->AssetId);
+			return FText::FromString(ModuleInfo.Name);
+		}
+ 	}
+	
+ 	return NSLOCTEXT("RewindDebugger", "AnimNextModuleTrackName", "Module");
+}
+
+void FAnimNextModuleTrackCreator::GetTrackTypesInternal(TArray<RewindDebugger::FRewindDebuggerTrackType>& Types) const
+{
+	Types.Add({AnimNextModulesName, LOCTEXT("AnimNextModule", "AnimNextModules")});
+}
+
+TSharedPtr<RewindDebugger::FRewindDebuggerTrack> FAnimNextModuleTrackCreator::CreateTrackInternal(uint64 ObjectId) const
+{
+	return MakeShared<FAnimNextModuleTrack>(ObjectId);
+}
+
+		
+FAnimNextModuleTrack::FAnimNextModuleTrack(uint64 InObjectId) :
+	ObjectId(InObjectId)
+{
+	FPropertyEditorModule& PropertyEditorModule = FModuleManager::GetModuleChecked<FPropertyEditorModule>("PropertyEditor");
+	FDetailsViewArgs DetailsViewArgs;
+	DetailsViewArgs.NameAreaSettings = FDetailsViewArgs::HideNameArea;
+	DetailsView = PropertyEditorModule.CreateDetailView(DetailsViewArgs);
+
+	InitializeDetailsObject();
+}
+	
+	
+FAnimNextModuleTrack::~FAnimNextModuleTrack()
+{
+ 	if (UPropertyBagDetailsObject* DetailsObject = DetailsObjectWeakPtr.Get())
+   	{
+   		DetailsObject->ClearFlags(RF_Standalone);
+   	}
+}
+	
+UPropertyBagDetailsObject* FAnimNextModuleTrack::InitializeDetailsObject()
+{
+	UPropertyBagDetailsObject* DetailsObject = NewObject<UPropertyBagDetailsObject>();
+	DetailsObject->SetFlags(RF_Standalone);
+	DetailsObjectWeakPtr = MakeWeakObjectPtr(DetailsObject);
+	DetailsView->SetObject(DetailsObject);
+	return DetailsObject;
+}
+
+
+bool FAnimNextModuleTrack::UpdateInternal()
+{
+	TRACE_CPUPROFILER_EVENT_SCOPE(FAnimNextModuleTrack::UpdateInternal);
+ 	
+	IRewindDebugger* RewindDebugger = IRewindDebugger::Instance();
+	
+	const TraceServices::IAnalysisSession* AnalysisSession = RewindDebugger->GetAnalysisSession();
+	
+	if (const FAnimNextProvider* AnimNextProvider = IRewindDebugger::Instance()->GetAnalysisSession()->ReadProvider<FAnimNextProvider>("AnimNextProvider"))
+	{
+		double CurrentScrubTime = IRewindDebugger::Instance()->CurrentTraceTime();
+		if (PreviousScrubTime != CurrentScrubTime)
+		{
+			PreviousScrubTime = CurrentScrubTime;
+
+			UPropertyBagDetailsObject* DetailsObject = DetailsObjectWeakPtr.Get();
+			if (DetailsObject == nullptr)
+			{
+				// this should not happen unless the object was garbage collected (which should not happen since it's marked as Standalone)
+				DetailsObject = InitializeDetailsObject();
+			}
+
+			DetailsObject->Properties.Reset();
+
+			if (!RewindDebugger->IsTraceFileLoaded()) // this currently crashes for object references on loaded trace files, so temporarily disabling it
+			{
+				const TraceServices::IFrameProvider& FramesProvider = TraceServices::ReadFrameProvider(*AnalysisSession);
+				TraceServices::FAnalysisSessionReadScope SessionReadScope(*AnalysisSession);
+			
+				TraceServices::FFrame MarkerFrame;
+				if(FramesProvider.GetFrameFromTime(ETraceFrameType::TraceFrameType_Game, CurrentScrubTime, MarkerFrame))
+				{
+					uint64 ModuleId;
+					if (AnimNextProvider->GetModuleId(ObjectId, ModuleId))
+					{
+						if (const FDataInterfaceData* Data = AnimNextProvider->GetDataInterfaceData(ModuleId))
+						{
+							Data->VariablesTimeline.EnumerateEvents(MarkerFrame.StartTime, MarkerFrame.EndTime, [DetailsObject](double InStartTime, double InEndTime, uint32 InDepth, const TArray<uint8>& VariableData)
+							{
+								FObjectReader Archive(VariableData);
+								DetailsObject->Properties.Serialize(Archive);
+							
+								return TraceServices::EEventEnumerate::Stop;
+							});
+						}
+					}
+				}
+			}
+		}
+	}
+	
+ 	return false;
+}
+
+void FAnimNextModuleTrack::IterateSubTracksInternal(TFunction<void(TSharedPtr<FRewindDebuggerTrack> SubTrack)> IteratorFunction)
+{
+};
+
+bool FAnimNextModuleTrack::HandleDoubleClickInternal()
+{
+	const IRewindDebugger* RewindDebugger = IRewindDebugger::Instance();
+	if (const TraceServices::IAnalysisSession* AnalysisSession = RewindDebugger->GetAnalysisSession())
+	{
+		TraceServices::FAnalysisSessionReadScope SessionReadScope(*AnalysisSession);
+
+		const IGameplayProvider* GameplayProvider = AnalysisSession->ReadProvider<IGameplayProvider>("GameplayProvider");
+		const FAnimNextProvider* AnimNextProvider = IRewindDebugger::Instance()->GetAnalysisSession()->ReadProvider<FAnimNextProvider>("AnimNextProvider");
+
+		uint64 ModuleId;
+		if (AnimNextProvider->GetModuleId(ObjectId, ModuleId))
+		{
+			if (const FDataInterfaceData* ModuleData = AnimNextProvider->GetDataInterfaceData(ModuleId))
+			{
+				const FObjectInfo& AssetInfo = GameplayProvider->GetObjectInfo(ModuleData->AssetId);
+				GEditor->GetEditorSubsystem<UAssetEditorSubsystem>()->OpenEditorForAsset(AssetInfo.PathName);
+				return true;
+			}
+		}
+	}
+	return false;
+}
+	
+	
+bool FAnimNextModuleTrackCreator::HasDebugInfoInternal(uint64 ObjectId) const
+{
+	TRACE_CPUPROFILER_EVENT_SCOPE(FAnimNextModuleTrack::HasDebugInfoInternal);
+
+	const IRewindDebugger* RewindDebugger = IRewindDebugger::Instance();
+	if (const TraceServices::IAnalysisSession* AnalysisSession = RewindDebugger->GetAnalysisSession())
+	{
+		TraceServices::FAnalysisSessionReadScope SessionReadScope(*AnalysisSession);
+    
+		const FAnimNextProvider* AnimNextProvider = IRewindDebugger::Instance()->GetAnalysisSession()->ReadProvider<FAnimNextProvider>("AnimNextProvider");
+    
+		uint64 ModuleId;
+		return AnimNextProvider->GetModuleId(ObjectId, ModuleId);
+	}
+
+	return false;
+}
+
+
+}
+
+#undef LOCTEXT_NAMESPACE
\ No newline at end of file
diff --git a/Engine/Plugins/Experimental/AnimNext/Source/AnimNextEditor/Private/RewindDebugger/AnimNextModuleTrack.h b/Engine/Plugins/Experimental/AnimNext/Source/AnimNextEditor/Private/RewindDebugger/AnimNextModuleTrack.h
new file mode 100644
--- /dev/null
+++ b/Engine/Plugins/Experimental/AnimNext/Source/AnimNextEditor/Private/RewindDebugger/AnimNextModuleTrack.h
@@ -0,0 +1,62 @@
+// Copyright Epic Games, Inc. All Rights Reserved.
+#pragma once
+
+#include "IRewindDebuggerTrackCreator.h"
+#include "RewindDebuggerTrack.h"
+#include "Textures/SlateIcon.h"
+#include "SEventTimelineView.h"
+#include "StructUtils/PropertyBag.h"
+#include "IDetailsView.h"
+
+#include "AnimNextModuleTrack.generated.h"
+
+UCLASS()
+class UPropertyBagDetailsObject : public UObject
+{
+	GENERATED_BODY()
+
+public:
+	UPROPERTY(VisibleAnywhere, Category=Details, meta=(ShowOnlyInnerProperties))
+	FInstancedPropertyBag Properties;
+};
+
+namespace UE::AnimNextEditor
+{
+	
+class FAnimNextModuleTrack : public RewindDebugger::FRewindDebuggerTrack
+{
+public:
+	FAnimNextModuleTrack(uint64 InObjectId);
+	virtual ~FAnimNextModuleTrack();
+private:
+	virtual bool UpdateInternal() override;
+	virtual TSharedPtr<SWidget> GetDetailsViewInternal() override { return DetailsView; }
+
+	virtual FSlateIcon GetIconInternal() override { return Icon; }
+	virtual FName GetNameInternal() const override { return "AnimNextModule"; }
+	virtual FText GetDisplayNameInternal() const override;
+	virtual uint64 GetObjectIdInternal() const override { return ObjectId; }
+	virtual void IterateSubTracksInternal(TFunction<void(TSharedPtr<FRewindDebuggerTrack> SubTrack)> IteratorFunction) override;
+	virtual bool HandleDoubleClickInternal() override;
+
+	UPropertyBagDetailsObject* InitializeDetailsObject();
+	
+	TSharedPtr<IDetailsView> DetailsView;
+	FSlateIcon Icon;
+	uint64 ObjectId;
+	double PreviousScrubTime = -1.0;
+	TWeakObjectPtr<UPropertyBagDetailsObject> DetailsObjectWeakPtr;
+};
+
+
+class FAnimNextModuleTrackCreator : public RewindDebugger::IRewindDebuggerTrackCreator
+{
+private:
+	virtual FName GetTargetTypeNameInternal() const;
+	virtual FName GetNameInternal() const override { return "AnimNextModule"; }
+	virtual void GetTrackTypesInternal(TArray<RewindDebugger::FRewindDebuggerTrackType>& Types) const override;
+	virtual TSharedPtr<RewindDebugger::FRewindDebuggerTrack> CreateTrackInternal(uint64 ObjectId) const override;
+	virtual bool HasDebugInfoInternal(uint64 ObjectId) const override;
+};
+
+}
\ No newline at end of file
diff --git a/Engine/Plugins/Experimental/AnimNext/Source/AnimNextEditor/Private/RewindDebugger/AnimNextProvider.cpp b/Engine/Plugins/Experimental/AnimNext/Source/AnimNextEditor/Private/RewindDebugger/AnimNextProvider.cpp
new file mode 100644
--- /dev/null
+++ b/Engine/Plugins/Experimental/AnimNext/Source/AnimNextEditor/Private/RewindDebugger/AnimNextProvider.cpp
@@ -0,0 +1,61 @@
+// Copyright Epic Games, Inc. All Rights Reserved.
+
+#include "AnimNextProvider.h"
+#include "ObjectTrace.h"
+
+FName FAnimNextProvider::ProviderName("AnimNextProvider");
+
+#define LOCTEXT_NAMESPACE "AnimNextProvider"
+
+FAnimNextProvider::FAnimNextProvider(TraceServices::IAnalysisSession& InSession)
+	: Session(InSession)
+{
+}
+
+void FAnimNextProvider::AppendModule(uint64 InModuleInstanceId, uint64 InModuleAssetId, uint64 InOuterObjectId)
+{
+	Session.WriteAccessCheck();
+
+	ComponentIdToModuleId.Add(InOuterObjectId, InModuleInstanceId);
+
+	TSharedRef<FDataInterfaceData> Data = MakeShared<FDataInterfaceData>(Session);
+	Data->InstanceId = InModuleInstanceId;
+	Data->AssetId = InModuleAssetId;
+
+	DataInterfaceData.Add( InModuleInstanceId, Data);
+}
+
+void FAnimNextProvider::AppendVariables(double ProfileTime, double RecordingTime, uint64 DataInterfaceId, const TArrayView<const uint8>& VariableData)
+{
+
+	Session.WriteAccessCheck();
+	
+	if (const TSharedRef<FDataInterfaceData>* Data = DataInterfaceData.Find(DataInterfaceId))
+	{
+		(*Data)->VariablesTimeline.AppendEvent(ProfileTime, TArray<uint8>(VariableData));
+	}
+}
+
+bool FAnimNextProvider::GetModuleId(uint64 ComponentId, uint64& OutModuleId) const
+{
+	Session.ReadAccessCheck();
+	
+	if (const uint64* Id = ComponentIdToModuleId.Find(ComponentId))
+	{
+		OutModuleId = *Id;
+		return true;
+	}
+	return false;
+}
+
+const FDataInterfaceData* FAnimNextProvider::GetDataInterfaceData(uint64 DataInterfaceId) const
+{
+	const TSharedRef<FDataInterfaceData>* Data = DataInterfaceData.Find(DataInterfaceId);
+	if (Data)
+	{
+		return &**Data;
+	}
+	return nullptr;
+}
+
+#undef LOCTEXT_NAMESPACE
diff --git a/Engine/Plugins/Experimental/AnimNext/Source/AnimNextEditor/Private/RewindDebugger/AnimNextProvider.h b/Engine/Plugins/Experimental/AnimNext/Source/AnimNextEditor/Private/RewindDebugger/AnimNextProvider.h
new file mode 100644
--- /dev/null
+++ b/Engine/Plugins/Experimental/AnimNext/Source/AnimNextEditor/Private/RewindDebugger/AnimNextProvider.h
@@ -0,0 +1,46 @@
+// Copyright Epic Games, Inc. All Rights Reserved.
+
+#pragma once
+
+#include "Engine/EngineBaseTypes.h"
+#include "Model/PointTimeline.h"
+#include "Containers/Map.h"
+#include "Templates/SharedPointer.h"
+#include "Model/IntervalTimeline.h"
+#include "TraceServices/Model/AnalysisSession.h"
+
+namespace TraceServices { class IAnalysisSession; }
+
+struct FDataInterfaceData
+{
+	FDataInterfaceData(TraceServices::IAnalysisSession& Session)
+		: VariablesTimeline(Session.GetLinearAllocator( )),
+			InstanceId (0),
+			AssetId (0)
+	{
+	}
+
+	TraceServices::TPointTimeline<TArray<uint8>> VariablesTimeline;
+	uint64 InstanceId;
+	uint64 AssetId;
+};
+
+class FAnimNextProvider : public TraceServices::IProvider
+{
+public:
+	static FName ProviderName;
+
+	FAnimNextProvider(TraceServices::IAnalysisSession& InSession);
+
+	void AppendModule(uint64 InModuleInstanceId, uint64 InModuleAssetId, uint64 InOuterObjectId);
+	void AppendVariables(double ProfileTime, double RecordingTime, uint64 InModuleInstanceId, const TArrayView<const uint8>& VariableData);
+
+	bool GetModuleId(uint64 ComponentId, uint64& OutModuleId) const;
+	const FDataInterfaceData* GetDataInterfaceData(uint64 DataInterfaceId) const;
+
+private:
+	TraceServices::IAnalysisSession& Session;
+
+	TMap<uint64, uint64> ComponentIdToModuleId;
+	TMap<uint64, TSharedRef<FDataInterfaceData>> DataInterfaceData;
+};
diff --git a/Engine/Plugins/Experimental/AnimNext/Source/AnimNextEditor/Private/RewindDebugger/AnimNextTraceModule.cpp b/Engine/Plugins/Experimental/AnimNext/Source/AnimNextEditor/Private/RewindDebugger/AnimNextTraceModule.cpp
new file mode 100644
--- /dev/null
+++ b/Engine/Plugins/Experimental/AnimNext/Source/AnimNextEditor/Private/RewindDebugger/AnimNextTraceModule.cpp
@@ -0,0 +1,32 @@
+// Copyright Epic Games, Inc. All Rights Reserved.
+
+#include "AnimNextTraceModule.h"
+#include "AnimNextProvider.h"
+#include "AnimNextAnalyzer.h"
+
+FName FAnimNextTraceModule::ModuleName("AnimNext");
+
+void FAnimNextTraceModule::GetModuleInfo(TraceServices::FModuleInfo& OutModuleInfo)
+{
+	OutModuleInfo.Name = ModuleName;
+	OutModuleInfo.DisplayName = TEXT("AnimNext");
+}
+
+void FAnimNextTraceModule::OnAnalysisBegin(TraceServices::IAnalysisSession& InSession)
+{
+	TSharedPtr<FAnimNextProvider> AnimNextProvider = MakeShared<FAnimNextProvider>(InSession);
+	InSession.AddProvider(FAnimNextProvider::ProviderName, AnimNextProvider);
+
+	InSession.AddAnalyzer(new FAnimNextAnalyzer(InSession, *AnimNextProvider));
+}
+
+void FAnimNextTraceModule::GetLoggers(TArray<const TCHAR *>& OutLoggers)
+{
+	OutLoggers.Add(TEXT("AnimNext"));
+}
+
+void FAnimNextTraceModule::GenerateReports(const TraceServices::IAnalysisSession& Session, const TCHAR* CmdLine, const TCHAR* OutputDirectory)
+{
+
+}
+
diff --git a/Engine/Plugins/Experimental/AnimNext/Source/AnimNextEditor/Private/RewindDebugger/AnimNextTraceModule.h b/Engine/Plugins/Experimental/AnimNext/Source/AnimNextEditor/Private/RewindDebugger/AnimNextTraceModule.h
new file mode 100644
--- /dev/null
+++ b/Engine/Plugins/Experimental/AnimNext/Source/AnimNextEditor/Private/RewindDebugger/AnimNextTraceModule.h
@@ -0,0 +1,23 @@
+// Copyright Epic Games, Inc. All Rights Reserved.
+
+#pragma once
+
+#include "CoreMinimal.h"
+#include "TraceServices/ModuleService.h"
+
+class IAnimationProvider;
+
+class FAnimNextTraceModule : public TraceServices::IModule
+{
+public:
+	// TraceServices::IModule interface
+	virtual void GetModuleInfo(TraceServices::FModuleInfo& OutModuleInfo) override;
+	virtual void OnAnalysisBegin(TraceServices::IAnalysisSession& Session) override;
+	virtual void GetLoggers(TArray<const TCHAR *>& OutLoggers) override;
+	virtual void GenerateReports(const TraceServices::IAnalysisSession& Session, const TCHAR* CmdLine, const TCHAR* OutputDirectory) override;
+	virtual const TCHAR* GetCommandLineArgument() override { return TEXT("chooser"); }
+
+private:
+	static FName ModuleName;
+};
+
diff --git a/Engine/Plugins/Experimental/AnimNext/Source/AnimNextEditor/Private/RewindDebugger/RewindDebuggerAnimNext.cpp b/Engine/Plugins/Experimental/AnimNext/Source/AnimNextEditor/Private/RewindDebugger/RewindDebuggerAnimNext.cpp
new file mode 100644
--- /dev/null
+++ b/Engine/Plugins/Experimental/AnimNext/Source/AnimNextEditor/Private/RewindDebugger/RewindDebuggerAnimNext.cpp
@@ -0,0 +1,15 @@
+// Copyright Epic Games, Inc. All Rights Reserved.
+
+#include "RewindDebuggerAnimNext.h"
+
+#include "IGameplayProvider.h"
+
+FRewindDebuggerAnimNext::FRewindDebuggerAnimNext()
+{
+
+}
+
+void FRewindDebuggerAnimNext::Update(float DeltaTime, IRewindDebugger* RewindDebugger)
+{
+	// Future place to sync anim next workspace editor debugging when scrubbing
+}
\ No newline at end of file
diff --git a/Engine/Plugins/Experimental/AnimNext/Source/AnimNextEditor/Private/RewindDebugger/RewindDebuggerAnimNext.h b/Engine/Plugins/Experimental/AnimNext/Source/AnimNextEditor/Private/RewindDebugger/RewindDebuggerAnimNext.h
new file mode 100644
--- /dev/null
+++ b/Engine/Plugins/Experimental/AnimNext/Source/AnimNextEditor/Private/RewindDebugger/RewindDebuggerAnimNext.h
@@ -0,0 +1,17 @@
+// Copyright Epic Games, Inc. All Rights Reserved.
+#pragma once
+
+#include "CoreMinimal.h"
+#include "IRewindDebugger.h"
+#include "IRewindDebuggerExtension.h"
+
+// Rewind debugger extension for Chooser support
+
+class FRewindDebuggerAnimNext : public IRewindDebuggerExtension
+{
+public:
+	FRewindDebuggerAnimNext();
+	virtual ~FRewindDebuggerAnimNext() {};
+
+	virtual void Update(float DeltaTime, IRewindDebugger* RewindDebugger) override;
+};
