Improve reference pose generation performance

diff --git a/Engine/Plugins/Experimental/AnimNext/Source/AnimNext/Private/GenerationTools.cpp b/Engine/Plugins/Experimental/AnimNext/Source/AnimNext/Private/GenerationTools.cpp
--- a/Engine/Plugins/Experimental/AnimNext/Source/AnimNext/Private/GenerationTools.cpp
+++ b/Engine/Plugins/Experimental/AnimNext/Source/AnimNext/Private/GenerationTools.cpp
@@ -324,26 +324,26 @@ void FGenerationTools::FixLODRequiredBones(const int32 NumLODs
 	, TArray<FGenerationLODData>& GenerationLODData
 	, TArray<FGenerationLODData>& GenerationComponentSpaceLODData)
 {
+	TArray<FBoneIndexType> MissingBones;
+
 	for (int32 LODIndex = NumLODs - 1; LODIndex > 0; --LODIndex)
 	{
+		MissingBones.Reset();
+
 		FGenerationLODData& LODData = GenerationLODData[LODIndex];
 		FGenerationLODData& PrevLODData = GenerationLODData[LODIndex - 1];
 
-
-		TArray<FBoneIndexType> MissingBones;
-		MissingBones.Reserve(LODData.RequiredBones.Num());
-
 		// Check if all required bones by the current LOD are part of the parent LOD as well.
 		// If a bone is present at a lower LOD while missing from the parent LOD, we will automatically add it to the parent LOD.
-		for (TArray<FBoneIndexType>::TConstIterator ReqBoneIt(LODData.RequiredBones); ReqBoneIt; ++ReqBoneIt)
+		for (FBoneIndexType BoneIdx : LODData.RequiredBones)
 		{
-			if (!PrevLODData.RequiredBones.Contains(*ReqBoneIt))
+			if (!PrevLODData.RequiredBones.Contains(BoneIdx))
 			{
-				MissingBones.Add(*ReqBoneIt);
+				MissingBones.Add(BoneIdx);
 
 				UE_LOG(LogAnimGenerationTools, Warning, TEXT("SkeletalMesh LOD %d does not contain bone [%s] required by LOD %d. Please update the skeletal mesh asset [%s] or its corresponding LOD settings asset.")
 					, LODIndex - 1
-					, *GetBoneNameSafe(SkeletalMesh, *ReqBoneIt)
+					, *GetBoneNameSafe(SkeletalMesh, BoneIdx)
 					, LODIndex
 					, *SkeletalMesh->GetPathName());
 			}
@@ -361,7 +361,7 @@ void FGenerationTools::FixLODRequiredBones(const int32 NumLODs
 				const int32 NumElem = MissingBones.Num();
 				for (int32 i = NumElem - 1; i >= 0; --i)
 				{
-					const FBoneIndexType& BoneIndex = MissingBones[i];
+					const FBoneIndexType BoneIndex = MissingBones[i];
 
 					PrevLODData.ExcludedBones.Remove(BoneIndex);
 					PrevLODData.ExcludedBonesFromPrevLOD.Remove(BoneIndex);
