AnimNext: Ensure debug draw primitive owner object is set up correctly

diff --git a/Engine/Plugins/Experimental/AnimNext/Source/AnimNext/Private/AnimNextDebugDraw.cpp b/Engine/Plugins/Experimental/AnimNext/Source/AnimNext/Private/AnimNextDebugDraw.cpp
--- a/Engine/Plugins/Experimental/AnimNext/Source/AnimNext/Private/AnimNextDebugDraw.cpp
+++ b/Engine/Plugins/Experimental/AnimNext/Source/AnimNext/Private/AnimNextDebugDraw.cpp
@@ -145,6 +145,7 @@ FDebugDraw::FDebugDraw(UObject* InOwner)
 	SceneInfoData->OwnerLastRenderTimePtr = nullptr;
 	SceneDesc.ProxyDesc = &SceneProxyDesc;
 	SceneDesc.PrimitiveSceneData = SceneInfoData;
+	SceneDesc.PrimitiveUObject = InOwner;
 
 	SceneDesc.RenderMatrix = FMatrix::Identity;
 	SceneDesc.AttachmentRootPosition = FVector::ZeroVector;
