Fix for multithreading crash with AnimNext tracing

diff --git a/Engine/Plugins/Experimental/AnimNext/Source/AnimNext/Private/RewindDebugger/AnimNextTrace.cpp b/Engine/Plugins/Experimental/AnimNext/Source/AnimNext/Private/RewindDebugger/AnimNextTrace.cpp
--- a/Engine/Plugins/Experimental/AnimNext/Source/AnimNext/Private/RewindDebugger/AnimNextTrace.cpp
+++ b/Engine/Plugins/Experimental/AnimNext/Source/AnimNext/Private/RewindDebugger/AnimNextTrace.cpp
@@ -28,6 +28,7 @@ namespace UE::AnimNext
 {
 	
 TSet<uint64> GTracedModules;
+FCriticalSection GTracedModulesLock;
 
 void FAnimNextTrace::Reset()
 {
@@ -47,11 +48,21 @@ void FAnimNextTrace::OutputAnimNextModuleInstance(const FAnimNextModuleInstance*
 		return;
 	}
 	uint64 ModuleId = ModuleInstance->GetHandle().GetUniqueId();
+	
+	bool bTrace = false;
 
-	if (!GTracedModules.Contains(ModuleId))
 	{
-		GTracedModules.Add(ModuleId);
+		FScopeLock Lock(&GTracedModulesLock);
+		
+		if (!GTracedModules.Contains(ModuleId))
+		{
+			GTracedModules.Add(ModuleId);
+			bTrace = true;
+		}
+	}
 
+	if (bTrace)
+	{
 		TRACE_OBJECT(ModuleInstance->GetObject());
 		TRACE_OBJECT(ModuleInstance->GetModule());
 
