Fall back to Host instance when looking up variables with FAnimNextDataInterfaceInstance::GetVariableInternal

diff --git a/Engine/Plugins/Experimental/AnimNext/Source/AnimNext/Private/DataInterface/AnimNextDataInterfaceInstance.cpp b/Engine/Plugins/Experimental/AnimNext/Source/AnimNext/Private/DataInterface/AnimNextDataInterfaceInstance.cpp
--- a/Engine/Plugins/Experimental/AnimNext/Source/AnimNext/Private/DataInterface/AnimNextDataInterfaceInstance.cpp
+++ b/Engine/Plugins/Experimental/AnimNext/Source/AnimNext/Private/DataInterface/AnimNextDataInterfaceInstance.cpp
@@ -568,7 +568,14 @@ EPropertyBagResult FAnimNextDataInterfaceInstance::GetVariableInternal(FName InV
 	const FPropertyBagPropertyDesc* Desc = PropertyBag->FindPropertyDescByName(InVariableName);
 	if(Desc == nullptr)
 	{
-		return EPropertyBagResult::PropertyNotFound;
+		if(HostInstance)
+		{
+			return HostInstance->GetVariableInternal(InVariableName, InType, OutResult);
+		}
+		else
+		{
+			return EPropertyBagResult::PropertyNotFound;
+		}
 	}
 
 	const int32 DescIndex = Desc - &PropertyBag->GetPropertyDescs()[0];
