Skeletal mesh attachments with minimal bones

This allows us to remap and attach skeletal meshes that only contain the bones it needs, like e.g. only skinned bones.

For the cases where the target skeleton does not share the same root bone with the source, try to find the corresponding bone on the source skeleton for the target's root bone and sync the skeletons up from there by calculating the delta transform between the two and move it along with the source.

diff --git a/Engine/Plugins/Experimental/AnimNext/Source/AnimNext/Private/Graph/RemapPoseData.cpp b/Engine/Plugins/Experimental/AnimNext/Source/AnimNext/Private/Graph/RemapPoseData.cpp
--- a/Engine/Plugins/Experimental/AnimNext/Source/AnimNext/Private/Graph/RemapPoseData.cpp
+++ b/Engine/Plugins/Experimental/AnimNext/Source/AnimNext/Private/Graph/RemapPoseData.cpp
@@ -55,6 +55,8 @@ void FRemapPoseData::Reinit(const UE::AnimNext::FReferencePose& InSourceRefPose,
 		const TArrayView<const FBoneIndexType> TargetPoseToMeshBoneIndexMap = InTargetRefPose.GetLODBoneIndexToMeshBoneIndexMap(LODLevel); // pose -> mesh
 		const int32 SourceNumBonesInLOD = SourceRefPose->GetNumBonesForLOD(LODLevel);
 
+		FBoneRemapping& SourceToTargetBoneIndexMap = SourceToTargetBoneIndexMapPerLOD[LODLevel];
+
 		// Iterate over the bones present in the given LOD on the target skeleton.
 		// These bones are candidates for our mapping table, only the bones from the target skeleton that are available in the source skeleton are of interest.
 		for (FBoneIndexType BoneIndex = 0; BoneIndex < TargetPoseToMeshBoneIndexMap.Num(); ++BoneIndex)
@@ -76,11 +78,43 @@ void FRemapPoseData::Reinit(const UE::AnimNext::FReferencePose& InSourceRefPose,
 					RemappedBone.SourceBoneIndex = RemappedIndex;
 					RemappedBone.TargetBoneIndex = BoneIndex;
 
-					SourceToTargetBoneIndexMapPerLOD[LODLevel].BoneIndexMap.Add(RemappedBone);
+					SourceToTargetBoneIndexMap.BoneIndexMap.Add(RemappedBone);
 				}
 			}
 		}
+
+		for (int32 MappingElementIndex = 0; MappingElementIndex < SourceToTargetBoneIndexMap.BoneIndexMap.Num(); ++MappingElementIndex)
+		{
+			const int32 SourceBoneIndex = SourceToTargetBoneIndexMap.BoneIndexMap[MappingElementIndex].SourceBoneIndex;
+			const int32 TargetBoneIndex = SourceToTargetBoneIndexMap.BoneIndexMap[MappingElementIndex].TargetBoneIndex;
+
+			if (TargetBoneIndex == UE::AnimNext::FReferencePose::RootBoneIndex)
+			{
+				SourceToTargetBoneIndexMap.TargetRootToSourceBoneIndex = SourceBoneIndex;
+				break;
+			}
+		}
+	}
+}
+
+FTransform FRemapPoseData::RecursiveCalcModelspaceTransform(const UE::AnimNext::FLODPoseHeap& Pose, FBoneIndexType BoneIndex) const
+{
+	if (BoneIndex == (uint16)INDEX_NONE)
+	{
+		return FTransform::Identity;
+	}
+
+	const FTransform LocalTransform = Pose.LocalTransforms[BoneIndex];
+
+	const TArrayView<const FBoneIndexType> BoneToParentBoneIndexMap = Pose.GetLODBoneIndexToParentLODBoneIndexMap();
+	FBoneIndexType ParentBoneIndex = BoneToParentBoneIndexMap[BoneIndex];
+	if (ParentBoneIndex == (uint16)INDEX_NONE)
+	{
+		return LocalTransform;
 	}
+
+	const FTransform ParentModelTransform = RecursiveCalcModelspaceTransform(Pose, ParentBoneIndex);
+	return LocalTransform * ParentModelTransform;
 }
 
 // TODO: Source pose dictates the LOD level of the target at the moment. What if the LOD levels of source and target are not in sync? We probably need to use the target skel mesh component LOD level when preparing the target pose.
@@ -108,6 +142,18 @@ void FRemapPoseData::RemapPose(const UE::AnimNext::FLODPoseHeap& SourcePose,
 		const auto& TargetTransform = SourcePose.LocalTransformsView[SourceBoneIndex];
 		OutTargetPose.LocalTransforms[TargetBoneIndex] = TargetTransform;
 	}
+
+	// For the cases where the target skeleton does not share the same root bone with the source, try to find the corresponding
+	// bone on the source skeleton for the target's root bone and sync the skeletons up from there by calculating the delta transform
+	// between the two and move it along with the source. This allows us to remap and attach skeletal meshes that only contain the
+	// bones it needs, like e.g. only skinned bones.
+	const int32 RootBoneOnSourceMapped = SourceToTargetBoneIndexMapPerLOD[SourceLODLevel].TargetRootToSourceBoneIndex;
+	if (RootBoneOnSourceMapped != INDEX_NONE)
+	{
+		// TODO: Replace with the generic method once we have access to the modelspace transform via the LOD pose.
+		const FTransform RootModelTransform = RecursiveCalcModelspaceTransform(SourcePose, RootBoneOnSourceMapped);
+		OutTargetPose.LocalTransforms[UE::AnimNext::FReferencePose::RootBoneIndex] = RootModelTransform;
+	}
 }
 
 void FRemapPoseData::RemapAttributes(const UE::AnimNext::FLODPose& SourceLODPose,
diff --git a/Engine/Plugins/Experimental/AnimNext/Source/AnimNext/Private/Graph/RemapPoseData.h b/Engine/Plugins/Experimental/AnimNext/Source/AnimNext/Private/Graph/RemapPoseData.h
--- a/Engine/Plugins/Experimental/AnimNext/Source/AnimNext/Private/Graph/RemapPoseData.h
+++ b/Engine/Plugins/Experimental/AnimNext/Source/AnimNext/Private/Graph/RemapPoseData.h
@@ -27,6 +27,12 @@ struct FBoneRemapping
 	 * This is a subset of the bones from the target pose. The smaller the subset, the less bones to remap.
 	 */
 	TArray<FRemappedBone> BoneIndexMap;
+
+	/**
+	 * Target root bone with index 0 mapped to the source skeletal mesh. The bone index stored here represents the bone index on the source
+	 * skeletal mesh of the bone with the same name as the root bone on the target. This is used as bone attachment point.
+	 */
+	int32 TargetRootToSourceBoneIndex = INDEX_NONE;
 };
 
 USTRUCT()
@@ -59,6 +65,11 @@ struct FRemapPoseData
 		const UE::AnimNext::FLODPose& TargetLODPose,
 		UE::Anim::FHeapAttributeContainer& OutAttributes);
 
+	/**
+	 * Temporary calculate modelspace transform, until the LODPose is able to do that.
+	 */
+	FTransform RecursiveCalcModelspaceTransform(const UE::AnimNext::FLODPoseHeap& Pose, FBoneIndexType BoneIndex) const;
+
 	const UE::AnimNext::FReferencePose* SourceRefPose = nullptr;
 	const UE::AnimNext::FReferencePose* TargetRefPose = nullptr;
 
diff --git a/Engine/Plugins/Experimental/AnimNext/Source/AnimNext/Public/ReferencePose.h b/Engine/Plugins/Experimental/AnimNext/Source/AnimNext/Public/ReferencePose.h
--- a/Engine/Plugins/Experimental/AnimNext/Source/AnimNext/Public/ReferencePose.h
+++ b/Engine/Plugins/Experimental/AnimNext/Source/AnimNext/Public/ReferencePose.h
@@ -71,6 +71,8 @@ struct TReferencePose
 	TWeakObjectPtr<const USkeleton> Skeleton = nullptr;
 	EReferencePoseGenerationFlags GenerationFlags = EReferencePoseGenerationFlags::None;
 
+	static constexpr FBoneIndexType RootBoneIndex = 0;
+
 	TReferencePose() = default;
 
 	bool IsValid() const
