AnimNext: Fix for leftover root motion from previous updates

diff --git a/Engine/Plugins/Experimental/AnimNext/Source/AnimNext/Private/Graph/RigUnit_AnimNextWriteSkeletalMeshComponentPose.cpp b/Engine/Plugins/Experimental/AnimNext/Source/AnimNext/Private/Graph/RigUnit_AnimNextWriteSkeletalMeshComponentPose.cpp
--- a/Engine/Plugins/Experimental/AnimNext/Source/AnimNext/Private/Graph/RigUnit_AnimNextWriteSkeletalMeshComponentPose.cpp
+++ b/Engine/Plugins/Experimental/AnimNext/Source/AnimNext/Private/Graph/RigUnit_AnimNextWriteSkeletalMeshComponentPose.cpp
@@ -92,6 +92,11 @@ FRigUnit_AnimNextWriteSkeletalMeshComponentPose_Execute()
 		}
 	}
 
+	// Clear our curves and attributes or we'll have leftovers from our last write.
+	OutputComponent->AnimCurves.Empty();
+	OutputComponent->GetEditableCustomAttributes().Empty();
+
+
 	// Map LOD pose into local-space scratch buffer
 	FGenerationTools::RemapPose(Pose.LODPose, LocalSpaceTransforms);
 
