Crash when remap target has less LODs than the source skeletal mesh

In case the source ref pose has fewer LOD levels available than the input/source pose, use the highest quality mapping available. This happens in case the source skeletal mesh does e.g. have 4 LODs while the target skeletal mesh only has 1. In this case we just create a single mapping for memory reasons as we'd store duplicated data otherwise and bones are guaranteed to be sorted.

diff --git a/Engine/Plugins/Experimental/AnimNext/Source/AnimNext/Private/Graph/RemapPoseData.cpp b/Engine/Plugins/Experimental/AnimNext/Source/AnimNext/Private/Graph/RemapPoseData.cpp
--- a/Engine/Plugins/Experimental/AnimNext/Source/AnimNext/Private/Graph/RemapPoseData.cpp
+++ b/Engine/Plugins/Experimental/AnimNext/Source/AnimNext/Private/Graph/RemapPoseData.cpp
@@ -40,58 +40,64 @@ void FRemapPoseData::Reinit(const UE::AnimNext::FReferencePose& InSourceRefPose,
 	const USkeletalMesh* SourceSkeletalMesh = InSourceRefPose.SkeletalMesh.Get();
 	const USkeleton* SourceSkeleton = SourceSkeletalMesh->GetSkeleton();
 	const FReferenceSkeleton& SourceRefSkeleton = SourceSkeleton->GetReferenceSkeleton();
+	const int32 SourceNumLODLevels = SourceSkeletalMesh->GetLODNum();
 
 	const USkeletalMesh* TargetSkeletalMesh = InTargetRefPose.SkeletalMesh.Get();
 	const FReferenceSkeleton& TargetRefSkeleton = TargetSkeletalMesh->GetRefSkeleton();
 	const int32 TargetNumLODLevels = TargetSkeletalMesh->GetLODNum();
 
+	// First index is the source LOD
 	SourceToTargetBoneIndexMapPerLOD.Empty();
-	SourceToTargetBoneIndexMapPerLOD.AddDefaulted(TargetNumLODLevels);
+	SourceToTargetBoneIndexMapPerLOD.AddDefaulted(SourceNumLODLevels);
 
-	const TArrayView<const FBoneIndexType> SourceSkeletonIndexToPoseIndex = InSourceRefPose.GetSkeletonBoneIndexToLODBoneIndexMap(); // skeleton -> pose
-
-	for (int32 LODLevel = 0; LODLevel < TargetNumLODLevels; ++LODLevel)
+	for (int32 SourceLODLevel = 0; SourceLODLevel < SourceNumLODLevels; ++SourceLODLevel)
 	{
-		const TArrayView<const FBoneIndexType> TargetPoseToMeshBoneIndexMap = InTargetRefPose.GetLODBoneIndexToMeshBoneIndexMap(LODLevel); // pose -> mesh
-		const int32 SourceNumBonesInLOD = SourceRefPose->GetNumBonesForLOD(LODLevel);
+		const TArrayView<const FBoneIndexType> SourceSkeletonIndexToPoseIndex = InSourceRefPose.GetSkeletonBoneIndexToLODBoneIndexMap(); // skeleton -> pose
+		const int32 SourceNumBonesInLOD = SourceRefPose->GetNumBonesForLOD(SourceLODLevel);
 
-		FBoneRemapping& SourceToTargetBoneIndexMap = SourceToTargetBoneIndexMapPerLOD[LODLevel];
+		// Second index is the target LOD
+		SourceToTargetBoneIndexMapPerLOD[SourceLODLevel].AddDefaulted(TargetNumLODLevels);
 
-		// Iterate over the bones present in the given LOD on the target skeleton.
-		// These bones are candidates for our mapping table, only the bones from the target skeleton that are available in the source skeleton are of interest.
-		for (FBoneIndexType BoneIndex = 0; BoneIndex < TargetPoseToMeshBoneIndexMap.Num(); ++BoneIndex)
+		for (int32 TargetLODLevel = 0; TargetLODLevel < TargetNumLODLevels; ++TargetLODLevel)
 		{
-			const FName BoneName = TargetRefSkeleton.GetBoneName(TargetPoseToMeshBoneIndexMap[BoneIndex]);
-			const int32 SourceSkeletonBoneIndex = SourceRefSkeleton.FindBoneIndex(BoneName);
+			const TArrayView<const FBoneIndexType> TargetPoseToMeshBoneIndexMap = InTargetRefPose.GetLODBoneIndexToMeshBoneIndexMap(TargetLODLevel); // pose -> mesh
+			FBoneRemapping& SourceToTargetBoneIndexMap = SourceToTargetBoneIndexMapPerLOD[SourceLODLevel][TargetLODLevel];
 
-			// Is the current bone on the target skeleton part of the source skeleton as well? Only add bones that are present on the target.
-			if (SourceSkeletonBoneIndex != INDEX_NONE)
+			// Iterate over the bones present in the given LOD on the target skeleton.
+			// These bones are candidates for our mapping table, only the bones from the target skeleton that are available in the source skeleton are of interest.
+			for (FBoneIndexType BoneIndex = 0; BoneIndex < TargetPoseToMeshBoneIndexMap.Num(); ++BoneIndex)
 			{
-				// TODO: Do we need to do this at runtime? This might fail if the LOD level of the source and the target differ. We either move this to the runtime pose remapping or store a matrix of remappings.
-				const FBoneIndexType RemappedIndex = SourceSkeletonIndexToPoseIndex[SourceSkeletonBoneIndex];
+				const FName BoneName = TargetRefSkeleton.GetBoneName(TargetPoseToMeshBoneIndexMap[BoneIndex]);
+				const int32 SourceSkeletonBoneIndex = SourceRefSkeleton.FindBoneIndex(BoneName);
 
-				// There are cases where the given bone are present in the skeleton, but not part of the actual pose/skeletal mesh,
-				// like e.g. when they are disabled for the given LOD on the target while they are not on the source. Skip these bones.
-				if (RemappedIndex != INDEX_NONE && RemappedIndex < SourceNumBonesInLOD)
+				// Is the current bone on the target skeleton part of the source skeleton as well? Only add bones that are present on the target.
+				if (SourceSkeletonBoneIndex != INDEX_NONE)
 				{
-					FRemappedBone RemappedBone;
-					RemappedBone.SourceBoneIndex = RemappedIndex;
-					RemappedBone.TargetBoneIndex = BoneIndex;
-
-					SourceToTargetBoneIndexMap.BoneIndexMap.Add(RemappedBone);
+					const FBoneIndexType RemappedIndex = SourceSkeletonIndexToPoseIndex[SourceSkeletonBoneIndex];
+
+					// There are cases where the given bone are present in the skeleton, but not part of the actual pose/skeletal mesh,
+					// like e.g. when they are disabled for the given LOD on the target while they are not on the source. Skip these bones.
+					if (RemappedIndex != INDEX_NONE && RemappedIndex < SourceNumBonesInLOD)
+					{
+						FRemappedBone RemappedBone;
+						RemappedBone.SourceBoneIndex = RemappedIndex;
+						RemappedBone.TargetBoneIndex = BoneIndex;
+
+						SourceToTargetBoneIndexMap.BoneIndexMap.Add(RemappedBone);
+					}
 				}
 			}
-		}
-
-		for (int32 MappingElementIndex = 0; MappingElementIndex < SourceToTargetBoneIndexMap.BoneIndexMap.Num(); ++MappingElementIndex)
-		{
-			const int32 SourceBoneIndex = SourceToTargetBoneIndexMap.BoneIndexMap[MappingElementIndex].SourceBoneIndex;
-			const int32 TargetBoneIndex = SourceToTargetBoneIndexMap.BoneIndexMap[MappingElementIndex].TargetBoneIndex;
 
-			if (TargetBoneIndex == UE::AnimNext::FReferencePose::RootBoneIndex)
+			for (int32 MappingElementIndex = 0; MappingElementIndex < SourceToTargetBoneIndexMap.BoneIndexMap.Num(); ++MappingElementIndex)
 			{
-				SourceToTargetBoneIndexMap.TargetRootToSourceBoneIndex = SourceBoneIndex;
-				break;
+				const int32 SourceBoneIndex = SourceToTargetBoneIndexMap.BoneIndexMap[MappingElementIndex].SourceBoneIndex;
+				const int32 TargetBoneIndex = SourceToTargetBoneIndexMap.BoneIndexMap[MappingElementIndex].TargetBoneIndex;
+
+				if (TargetBoneIndex == UE::AnimNext::FReferencePose::RootBoneIndex)
+				{
+					SourceToTargetBoneIndexMap.TargetRootToSourceBoneIndex = SourceBoneIndex;
+					break;
+				}
 			}
 		}
 	}
@@ -133,7 +139,12 @@ void FRemapPoseData::RemapPose(const UE::AnimNext::FLODPoseHeap& SourcePose,
 
 	OutTargetPose.Flags = SourcePose.Flags;
 
-	const TArray<FRemappedBone>& SourceToTargetBoneIndexMap = SourceToTargetBoneIndexMapPerLOD[SourceLODLevel].BoneIndexMap;
+	// In case the source ref pose has fewer LOD levels available than the input/source pose, use the highest quality mapping available.
+	// This happens in case the source skeletal mesh does e.g. have 4 LODs while the target skeletal mesh only has 1. In this case we just create a single mapping
+	// for memory reasons as we'd store duplicated data otherwise and bones are guaranteed to be sorted.
+	const int32 TargetLODIndex = FMath::Clamp(SourceLODLevel, 0, SourceToTargetBoneIndexMapPerLOD[SourceLODLevel].Num() - 1);
+	const TArray<FRemappedBone>& SourceToTargetBoneIndexMap = SourceToTargetBoneIndexMapPerLOD[SourceLODLevel][TargetLODIndex].BoneIndexMap;
+
 	for (int32 MappingElementIndex = 0; MappingElementIndex < SourceToTargetBoneIndexMap.Num(); ++MappingElementIndex)
 	{
 		const int32 SourceBoneIndex = SourceToTargetBoneIndexMap[MappingElementIndex].SourceBoneIndex;
@@ -147,7 +158,7 @@ void FRemapPoseData::RemapPose(const UE::AnimNext::FLODPoseHeap& SourcePose,
 	// bone on the source skeleton for the target's root bone and sync the skeletons up from there by calculating the delta transform
 	// between the two and move it along with the source. This allows us to remap and attach skeletal meshes that only contain the
 	// bones it needs, like e.g. only skinned bones.
-	const int32 RootBoneOnSourceMapped = SourceToTargetBoneIndexMapPerLOD[SourceLODLevel].TargetRootToSourceBoneIndex;
+	const int32 RootBoneOnSourceMapped = SourceToTargetBoneIndexMapPerLOD[SourceLODLevel][TargetLODIndex].TargetRootToSourceBoneIndex;
 	if (RootBoneOnSourceMapped != INDEX_NONE)
 	{
 		// TODO: Replace with the generic method once we have access to the modelspace transform via the LOD pose.
diff --git a/Engine/Plugins/Experimental/AnimNext/Source/AnimNext/Private/Graph/RemapPoseData.h b/Engine/Plugins/Experimental/AnimNext/Source/AnimNext/Private/Graph/RemapPoseData.h
--- a/Engine/Plugins/Experimental/AnimNext/Source/AnimNext/Private/Graph/RemapPoseData.h
+++ b/Engine/Plugins/Experimental/AnimNext/Source/AnimNext/Private/Graph/RemapPoseData.h
@@ -73,6 +73,13 @@ struct FRemapPoseData
 	const UE::AnimNext::FReferencePose* SourceRefPose = nullptr;
 	const UE::AnimNext::FReferencePose* TargetRefPose = nullptr;
 
-	/** Bone mapping for each LOD level. */
-	TArray<FBoneRemapping> SourceToTargetBoneIndexMapPerLOD;
+	/**
+	 * Bone mapping from the source to the target skeletal mesh per LOD level combination.
+	 * The array's length equals the number of source skeletal mesh LODs available. Each of the elements
+	 * will contain a bone mapping based on the given source LOD for the target skeletal mesh LOD. The number of elements
+	 * of the sub-array equals to the number of LOD levels of the target skeletal mesh.
+	 * Example: SourceToTargetBoneIndexMapPerLOD[4][2] gives a bone mapping from source LOD 4 to target LOD 2 and will usually contain
+	 * less elements than the target skeletal mesh has bones.
+	 */
+	TArray<TArray<FBoneRemapping>> SourceToTargetBoneIndexMapPerLOD;
 };
\ No newline at end of file
