AnimNext: Fix crash when disabling the initialize event when an instance has already been removed

diff --git a/Engine/Plugins/Experimental/AnimNext/Source/AnimNext/Private/Module/RigUnit_AnimNextModuleEvents.cpp b/Engine/Plugins/Experimental/AnimNext/Source/AnimNext/Private/Module/RigUnit_AnimNextModuleEvents.cpp
--- a/Engine/Plugins/Experimental/AnimNext/Source/AnimNext/Private/Module/RigUnit_AnimNextModuleEvents.cpp
+++ b/Engine/Plugins/Experimental/AnimNext/Source/AnimNext/Private/Module/RigUnit_AnimNextModuleEvents.cpp
@@ -77,7 +77,16 @@ FRigUnit_AnimNextInitializeEvent_Execute()
 	// Queue a GT task to disable ourselves (we can only enable/disable tick functions on the GT)
 	FAnimNextModuleInstance::RunTaskOnGameThread([ModuleHandle = ModuleInstance.Handle, WeakObject = TWeakObjectPtr<UObject>(ModuleInstance.Object)]()
 	{
-		UAnimNextWorldSubsystem* WorldSubsystem = UAnimNextWorldSubsystem::Get(WeakObject.Get());
+		UObject* Object = WeakObject.Get();
+		if (Object == nullptr)
+		{
+			return;
+		}
+		UAnimNextWorldSubsystem* WorldSubsystem = UAnimNextWorldSubsystem::Get(Object);
+		if (WorldSubsystem == nullptr)
+		{
+			return;
+		}
 		WorldSubsystem->DisableInitializeEvent(ModuleHandle);
 	});
 }
