AnimNext: Fix race condition/crash when releasing debug draw when worker threads are still using it

diff --git a/Engine/Plugins/Experimental/AnimNext/Source/AnimNext/Internal/AnimNextDebugDraw.h b/Engine/Plugins/Experimental/AnimNext/Source/AnimNext/Internal/AnimNextDebugDraw.h
--- a/Engine/Plugins/Experimental/AnimNext/Source/AnimNext/Internal/AnimNextDebugDraw.h
+++ b/Engine/Plugins/Experimental/AnimNext/Source/AnimNext/Internal/AnimNextDebugDraw.h
@@ -11,6 +11,7 @@
 class UObject;
 struct FRigVMDrawInterface;
 struct FAnimNextModuleInstance;
+class UAnimNextWorldSubsystem;
 
 namespace UE::AnimNext
 {
@@ -44,7 +45,6 @@ struct FDebugDraw
 {
 public:
 	explicit FDebugDraw(UObject* InOwner);
-	~FDebugDraw();
 
 	// Draw any debug items in DrawInterface
 	void Draw();
@@ -66,6 +66,9 @@ private:
 
 	void CalcBounds();
 
+	// Unregister the primitive from the RT
+	void RemovePrimitive();
+
 	FCustomSceneProxyDesc SceneProxyDesc;
 	FPrimitiveSceneDesc SceneDesc;
 	FPrimitiveSceneInfoData* SceneInfoData = nullptr;
@@ -79,8 +82,15 @@ private:
 	// Anim thread accessible enabled flag
 	bool bIsEnabled = false;
 
+	// Flag indicating whether we are registered with the RT
+	bool bIsRegistered = false;
+
+	// Lock to prevent unregistration during WT drawing
+	FRWLock Lock;
+
 	friend struct ::FAnimNextModuleInstance;
 	friend struct UE::AnimNext::FModuleEventTickFunction;
+	friend class ::UAnimNextWorldSubsystem;
 };
 
 #endif
diff --git a/Engine/Plugins/Experimental/AnimNext/Source/AnimNext/Private/AnimNextDebugDraw.cpp b/Engine/Plugins/Experimental/AnimNext/Source/AnimNext/Private/AnimNextDebugDraw.cpp
--- a/Engine/Plugins/Experimental/AnimNext/Source/AnimNext/Private/AnimNextDebugDraw.cpp
+++ b/Engine/Plugins/Experimental/AnimNext/Source/AnimNext/Private/AnimNextDebugDraw.cpp
@@ -140,6 +140,8 @@ FDebugDraw::FDebugDraw(UObject* InOwner)
 	: SceneProxyDesc(InOwner)
 	, SceneProxy(new FAnimNextDebugSceneProxy(SceneProxyDesc))
 {
+	check(IsInGameThread());
+
 	SceneInfoData = new FPrimitiveSceneInfoData();
 	SceneInfoData->SceneProxy = SceneProxy;
 	SceneInfoData->OwnerLastRenderTimePtr = nullptr;
@@ -155,11 +157,15 @@ FDebugDraw::FDebugDraw(UObject* InOwner)
 	Scene = InOwner->GetWorld()->Scene;
 
 	Scene->AddPrimitive(&SceneDesc);
+	bIsRegistered = true;
 }
 
-FDebugDraw::~FDebugDraw()
+void FDebugDraw::RemovePrimitive()
 {
+	FWriteScopeLock WriteLock(Lock);
+
 	Scene->RemovePrimitive(&SceneDesc);
+	bIsRegistered = false;
 
 	// Delete the scene info data on the RT as we no longer own it
 	ENQUEUE_RENDER_COMMAND(AnimNextDebugDraw)([SceneInfoData = SceneInfoData](FRHICommandList& RHICmdList)
@@ -170,6 +176,13 @@ FDebugDraw::~FDebugDraw()
 
 void FDebugDraw::Draw()
 {
+	FWriteScopeLock WriteLock(Lock);
+
+	if (!bIsRegistered)
+	{
+		return;
+	}
+
 	// Recalc bounds according to what we need to draw
 	CalcBounds();
 
@@ -188,6 +201,8 @@ void FDebugDraw::Draw()
 
 void FDebugDraw::SetEnabled(bool bInIsEnabled)
 {
+	check(IsInGameThread());
+
 	if(bInIsEnabled != bIsEnabled)
 	{
 		// Update the RT copy of the flag if it differs
diff --git a/Engine/Plugins/Experimental/AnimNext/Source/AnimNext/Private/Component/AnimNextWorldSubsystem.cpp b/Engine/Plugins/Experimental/AnimNext/Source/AnimNext/Private/Component/AnimNextWorldSubsystem.cpp
--- a/Engine/Plugins/Experimental/AnimNext/Source/AnimNext/Private/Component/AnimNextWorldSubsystem.cpp
+++ b/Engine/Plugins/Experimental/AnimNext/Source/AnimNext/Private/Component/AnimNextWorldSubsystem.cpp
@@ -172,8 +172,8 @@ void UAnimNextWorldSubsystem::UnregisterHandle(UE::AnimNext::FModuleHandle& InOu
 		FAnimNextModuleInstance& Instance = Instances.Get(InOutHandle);
 		FModuleWriteGuard Guard(&Instance);
 #if UE_ENABLE_DEBUG_DRAWING
-		// Release debug drawing immediately as the renderer will need to know about this before EOF
-		Instance.DebugDraw.Reset();
+		// Remove debug drawing immediately as the renderer will need to know about this before EOF
+		Instance.DebugDraw->RemovePrimitive();
 #endif
 
 		// Remove all tick dependencies immediately, as once the handle has been invalidated there is no way for external systems to remove their dependencies
