Anim Next: Fix anim next component not getting set to enabled when starting PIE.

diff --git a/Engine/Plugins/Experimental/AnimNext/Source/AnimNext/Private/Component/AnimNextComponent.cpp b/Engine/Plugins/Experimental/AnimNext/Source/AnimNext/Private/Component/AnimNextComponent.cpp
--- a/Engine/Plugins/Experimental/AnimNext/Source/AnimNext/Private/Component/AnimNextComponent.cpp
+++ b/Engine/Plugins/Experimental/AnimNext/Source/AnimNext/Private/Component/AnimNextComponent.cpp
@@ -44,7 +44,11 @@ void UAnimNextComponent::BeginPlay()
 {
 	Super::BeginPlay();
 
-	if (InitMethod == EAnimNextModuleInitMethod::InitializeAndRun)
+	UWorld* World = GetWorld();
+	check(World);
+
+	if (InitMethod == EAnimNextModuleInitMethod::InitializeAndRun
+	|| (InitMethod == EAnimNextModuleInitMethod::InitializeAndPauseInEditor && World->WorldType != EWorldType::Editor && World->WorldType != EWorldType::EditorPreview))
 	{
 		SetEnabled(true);
 	}
