Anim Next Graph Instance tracing/debugging

diff --git a/Engine/Plugins/Experimental/AnimNext/Source/AnimNext/Internal/Module/AnimNextModuleInstance.h b/Engine/Plugins/Experimental/AnimNext/Source/AnimNext/Internal/Module/AnimNextModuleInstance.h
--- a/Engine/Plugins/Experimental/AnimNext/Source/AnimNext/Internal/Module/AnimNextModuleInstance.h
+++ b/Engine/Plugins/Experimental/AnimNext/Source/AnimNext/Internal/Module/AnimNextModuleInstance.h
@@ -165,6 +165,11 @@ private:
 	// Flip proxy variable buffers then copy any dirty values
 	void CopyProxyVariables();
 
+#if ANIMNEXT_TRACE_ENABLED
+	// trace this module if it hasn't already been traced this frame
+	ANIMNEXT_API void Trace();
+#endif
+
 #if WITH_EDITOR
 	// Resets internal state if the module we are bound to is recompiled in editor
 	void OnModuleCompiled();
@@ -301,6 +306,10 @@ private:
 private:
 	// Adds the specified component and returns a reference to it
 	ANIMNEXT_API FAnimNextModuleInstanceComponent& AddComponentInternal(int32 ComponentNameHash, FName ComponentName, TInstancedStruct<FAnimNextModuleInstanceComponent>&& Component);
+
+#if ANIMNEXT_TRACE_ENABLED
+	bool bTracedThisFrame = false;
+#endif
 };
 
 template<>
@@ -345,3 +354,9 @@ const ComponentType* FAnimNextModuleInstance::TryGetComponent() const
 
 	return static_cast<ComponentType*>(TryGetComponent(ComponentNameHash, ComponentName));
 }
+
+#if ANIMNEXT_TRACE_ENABLED
+#define TRACE_ANIMNEXT_MODULE(ModuleInstance) ModuleInstance.Trace();
+#else
+#define TRACE_ANIMNEXT_MODULE(ModuleInstance)
+#endif
\ No newline at end of file
diff --git a/Engine/Plugins/Experimental/AnimNext/Source/AnimNext/Private/DataInterface/AnimNextDataInterfaceInstance.cpp b/Engine/Plugins/Experimental/AnimNext/Source/AnimNext/Private/DataInterface/AnimNextDataInterfaceInstance.cpp
--- a/Engine/Plugins/Experimental/AnimNext/Source/AnimNext/Private/DataInterface/AnimNextDataInterfaceInstance.cpp
+++ b/Engine/Plugins/Experimental/AnimNext/Source/AnimNext/Private/DataInterface/AnimNextDataInterfaceInstance.cpp
@@ -4,6 +4,11 @@
 #include "Logging/StructuredLog.h"
 #include "Param/ParamUtils.h"
 
+#if ANIMNEXT_TRACE_ENABLED
+volatile int64 FAnimNextDataInterfaceInstance::NextUniqueId = 1;
+#endif
+
+
 namespace UE::AnimNext::Private
 {
 	// Utility functions for type conversion, adapted from UE::StructUtils::Private in PropertyBag.cpp
@@ -591,6 +596,13 @@ EPropertyBagResult FAnimNextDataInterfaceInstance::GetVariableInternal(FName InV
 	return EPropertyBagResult::Success;
 }
 
+FAnimNextDataInterfaceInstance::FAnimNextDataInterfaceInstance()
+{
+#if ANIMNEXT_TRACE_ENABLED
+	UniqueId = FPlatformAtomics::InterlockedIncrement(&NextUniqueId);
+#endif
+}
+
 uint8* FAnimNextDataInterfaceInstance::GetMemoryForVariable(int32 InVariableIndex, FName InVariableName, const FProperty* InVariableProperty) const
 {
 	const UPropertyBag* PropertyBag = Variables.GetPropertyBagStruct();
diff --git a/Engine/Plugins/Experimental/AnimNext/Source/AnimNext/Private/Module/AnimNextModuleInstance.cpp b/Engine/Plugins/Experimental/AnimNext/Source/AnimNext/Private/Module/AnimNextModuleInstance.cpp
--- a/Engine/Plugins/Experimental/AnimNext/Source/AnimNext/Private/Module/AnimNextModuleInstance.cpp
+++ b/Engine/Plugins/Experimental/AnimNext/Source/AnimNext/Private/Module/AnimNextModuleInstance.cpp
@@ -427,10 +427,23 @@ void FAnimNextModuleInstance::CopyProxyVariables()
 			PublicVariablesProxy.bIsDirty = false;
 		}
 	}
-	
-	TRACE_ANIMNEXT_MODULE_VARIABLES(this);
+
+#if ANIMNEXT_TRACE_ENABLED
+	bTracedThisFrame = false;
+#endif
 }
 
+#if ANIMNEXT_TRACE_ENABLED
+void FAnimNextModuleInstance::Trace()
+{
+	if (!bTracedThisFrame)
+	{
+		TRACE_ANIMNEXT_VARIABLES(this, Object);
+		bTracedThisFrame = true;
+	}
+}
+#endif
+
 const UAnimNextModule* FAnimNextModuleInstance::GetModule() const
 {
 	return CastChecked<UAnimNextModule>(DataInterface);
diff --git a/Engine/Plugins/Experimental/AnimNext/Source/AnimNext/Private/RewindDebugger/AnimNextTrace.cpp b/Engine/Plugins/Experimental/AnimNext/Source/AnimNext/Private/RewindDebugger/AnimNextTrace.cpp
--- a/Engine/Plugins/Experimental/AnimNext/Source/AnimNext/Private/RewindDebugger/AnimNextTrace.cpp
+++ b/Engine/Plugins/Experimental/AnimNext/Source/AnimNext/Private/RewindDebugger/AnimNextTrace.cpp
@@ -1,24 +1,27 @@
 // Copyright Epic Games, Inc. All Rights Reserved.
 
-#include "AnimNextTrace.h"
+#include "RewindDebugger/AnimNextTrace.h"
 
 #include "ObjectTrace.h"
+#include "Chaos/ParticleDirtyFlags.h"
+#include "DataInterface/AnimNextDataInterfaceInstance.h"
 #include "Module/AnimNextModuleInstance.h"
 #include "Serialization/ObjectWriter.h"
 
 #if ANIMNEXT_TRACE_ENABLED
 
 
-UE_TRACE_EVENT_BEGIN(AnimNext, Module)
-	UE_TRACE_EVENT_FIELD(uint64, ModuleInstanceId)
+UE_TRACE_EVENT_BEGIN(AnimNext, Instance)
+	UE_TRACE_EVENT_FIELD(uint64, InstanceId)
+	UE_TRACE_EVENT_FIELD(uint64, HostInstanceId)
 	UE_TRACE_EVENT_FIELD(uint64, OuterObjectId)
-	UE_TRACE_EVENT_FIELD(uint64, ModuleAssetId)
+	UE_TRACE_EVENT_FIELD(uint64, AssetId)
 UE_TRACE_EVENT_END()
 
-UE_TRACE_EVENT_BEGIN(AnimNext, ModuleVariables)
+UE_TRACE_EVENT_BEGIN(AnimNext, InstanceVariables)
 	UE_TRACE_EVENT_FIELD(uint64, Cycle)
 	UE_TRACE_EVENT_FIELD(double, RecordingTime)
-	UE_TRACE_EVENT_FIELD(uint64, ModuleInstanceId)
+	UE_TRACE_EVENT_FIELD(uint64, InstanceId)
 	UE_TRACE_EVENT_FIELD(uint8[], VariableData)
 UE_TRACE_EVENT_END()
 
@@ -26,82 +29,98 @@ UE_TRACE_CHANNEL_DEFINE(AnimNextChannel);
 
 namespace UE::AnimNext
 {
-	
-TSet<uint64> GTracedModules;
-FCriticalSection GTracedModulesLock;
+
+FCriticalSection GTracedInstancesLock;
+TSet<uint64> GTracedInstances;
 
 void FAnimNextTrace::Reset()
 {
-	GTracedModules.Empty();
+	GTracedInstances.Empty();
 }
 
-void FAnimNextTrace::OutputAnimNextModuleInstance(const FAnimNextModuleInstance* ModuleInstance)
+void FAnimNextTrace::OutputAnimNextInstance(const FAnimNextDataInterfaceInstance* DataInterface, const UObject* OuterObject)
 {
 	bool bChannelEnabled = UE_TRACE_CHANNELEXPR_IS_ENABLED(AnimNextChannel);
-	if (!bChannelEnabled || ModuleInstance == nullptr)
+	if (!bChannelEnabled || DataInterface == nullptr)
 	{
 		return;
 	}
 
-	if (CANNOT_TRACE_OBJECT(ModuleInstance->GetObject()))
+	if (CANNOT_TRACE_OBJECT(OuterObject))
 	{
 		return;
 	}
-	uint64 ModuleId = ModuleInstance->GetHandle().GetUniqueId();
 	
+	uint64 InstanceId = DataInterface->GetUniqueId();
+
 	bool bTrace = false;
 
 	{
-		FScopeLock Lock(&GTracedModulesLock);
+		FScopeLock Lock(&GTracedInstancesLock);
 		
-		if (!GTracedModules.Contains(ModuleId))
+		if (!GTracedInstances.Contains(InstanceId))
 		{
-			GTracedModules.Add(ModuleId);
+			GTracedInstances.Add(InstanceId);
 			bTrace = true;
 		}
 	}
 
 	if (bTrace)
 	{
-		TRACE_OBJECT(ModuleInstance->GetObject());
-		TRACE_OBJECT(ModuleInstance->GetModule());
+		if (DataInterface->GetHost())
+		{
+			OutputAnimNextInstance(DataInterface->GetHost(), OuterObject);
+		}
+		
+		TRACE_OBJECT(OuterObject);
+		TRACE_OBJECT(DataInterface->GetDataInterface());
+		
+		const FAnimNextDataInterfaceInstance* HostInstance = DataInterface->GetHost();
 
-		UE_TRACE_LOG(AnimNext, Module, AnimNextChannel)
-			<< Module.ModuleInstanceId(ModuleId)
-			<< Module.OuterObjectId(FObjectTrace::GetObjectId(ModuleInstance->GetObject()))
-			<< Module.ModuleAssetId(FObjectTrace::GetObjectId(ModuleInstance->GetModule()));
+		UE_TRACE_LOG(AnimNext, Instance, AnimNextChannel)
+			<< Instance.InstanceId(InstanceId)
+			<< Instance.OuterObjectId(FObjectTrace::GetObjectId(OuterObject))
+			<< Instance.HostInstanceId(HostInstance ? HostInstance->GetUniqueId() : 0)
+			<< Instance.AssetId(FObjectTrace::GetObjectId(DataInterface->GetDataInterface()));
 	}
 }
 
-void FAnimNextTrace::OutputAnimNextModuleInstanceVariables(const FAnimNextModuleInstance* ModuleInstance)
+void FAnimNextTrace::OutputAnimNextVariables(const FAnimNextDataInterfaceInstance* Instance, const UObject* OuterObject)
 {
 	bool bChannelEnabled = UE_TRACE_CHANNELEXPR_IS_ENABLED(AnimNextChannel);
-	if (!bChannelEnabled || ModuleInstance == nullptr)
+	if (!bChannelEnabled || Instance == nullptr || OuterObject == nullptr)
 	{
 		return;
 	}
 
-	if (CANNOT_TRACE_OBJECT(ModuleInstance->GetObject()))
+	uint64 InstanceId = Instance->GetUniqueId();
+
+	if (CANNOT_TRACE_OBJECT(OuterObject))
 	{
-		return;
+      	return;
 	}
-	uint64 ModuleId = ModuleInstance->GetHandle().GetUniqueId();
-
-	TRACE_ANIMNEXT_MODULE(ModuleInstance);
 
-	const FInstancedPropertyBag Variables = ModuleInstance->GetVariables();
+	OutputAnimNextInstance(Instance, OuterObject);
+	
+	FInstancedPropertyBag Variables = Instance->GetVariables();
+	TConstArrayView<FPropertyBagPropertyDesc> PropertyDescs = Variables.GetPropertyBagStruct()->GetPropertyDescs();
+	for(int i =0; i<PropertyDescs.Num(); i++)
+	{
+		uint8* ValueMemory = Instance->GetMemoryForVariable(i, PropertyDescs[i].Name, PropertyDescs[i].CachedProperty);
+		PropertyDescs[i].CachedProperty->CopySingleValue(const_cast<uint8*>(Variables.GetValue().GetMemory()) +  PropertyDescs[i].CachedProperty->GetOffset_ForInternal(), ValueMemory);
+	}
 
 	TArray<uint8> ArchiveData;
 	FObjectWriter Archive(ArchiveData);
 	const_cast<FInstancedPropertyBag&>(Variables).Serialize(Archive);
 
-	UE_TRACE_LOG(AnimNext, ModuleVariables, AnimNextChannel)
-			<< ModuleVariables.Cycle(FPlatformTime::Cycles64())
-			<< ModuleVariables.RecordingTime(FObjectTrace::GetWorldElapsedTime(ModuleInstance->GetObject()->GetWorld()))
-			<< ModuleVariables.ModuleInstanceId(ModuleId)
-			<< ModuleVariables.VariableData(ArchiveData.GetData(), ArchiveData.Num());
+	UE_TRACE_LOG(AnimNext, InstanceVariables, AnimNextChannel)
+			<< InstanceVariables.Cycle(FPlatformTime::Cycles64())
+			<< InstanceVariables.RecordingTime(FObjectTrace::GetWorldElapsedTime(OuterObject->GetWorld()))
+			<< InstanceVariables.InstanceId(InstanceId)
+			<< InstanceVariables.VariableData(ArchiveData.GetData(), ArchiveData.Num());
 }
-
+	
 }
 
 #endif
diff --git a/Engine/Plugins/Experimental/AnimNext/Source/AnimNext/Private/RewindDebugger/AnimNextTrace.h b/Engine/Plugins/Experimental/AnimNext/Source/AnimNext/Private/RewindDebugger/AnimNextTrace.h
deleted file mode 100644
diff --git a/Engine/Plugins/Experimental/AnimNext/Source/AnimNext/Public/DataInterface/AnimNextDataInterfaceInstance.h b/Engine/Plugins/Experimental/AnimNext/Source/AnimNext/Public/DataInterface/AnimNextDataInterfaceInstance.h
--- a/Engine/Plugins/Experimental/AnimNext/Source/AnimNext/Public/DataInterface/AnimNextDataInterfaceInstance.h
+++ b/Engine/Plugins/Experimental/AnimNext/Source/AnimNext/Public/DataInterface/AnimNextDataInterfaceInstance.h
@@ -7,6 +7,7 @@
 #include "Param/ParamType.h"
 #include "RigVMCore/RigVMExecuteContext.h"
 #include "StructUtils/PropertyBag.h"
+#include "RewindDebugger/AnimNextTrace.h"
 #include "AnimNextDataInterfaceInstance.generated.h"
 
 class UAnimNextDataInterface;
@@ -22,6 +23,8 @@ USTRUCT()
 struct ANIMNEXT_API FAnimNextDataInterfaceInstance
 {
 	GENERATED_BODY()
+	
+	FAnimNextDataInterfaceInstance();
 
 	// Get the data interface asset that this instance represents
 	const UAnimNextDataInterface* GetDataInterface() const
@@ -66,6 +69,13 @@ struct ANIMNEXT_API FAnimNextDataInterfaceInstance
 
 	// Get the instance (graph, module etc.) that owns/hosts us
 	FAnimNextDataInterfaceInstance* GetHost() const { return HostInstance; }
+	
+#if ANIMNEXT_TRACE_ENABLED
+	uint64 GetUniqueId() const
+	{
+		return UniqueId;
+	}
+#endif 
 
 private:
 	// Helper function for GetVariable
@@ -89,4 +99,9 @@ protected:
 
 	friend UE::AnimNext::FInjectionInfo;
 	friend FAnimNextModuleInjectionComponent;
+	
+#if ANIMNEXT_TRACE_ENABLED
+	uint64 UniqueId;
+	volatile static int64 NextUniqueId;
+#endif 
 };
diff --git a/Engine/Plugins/Experimental/AnimNext/Source/AnimNext/Public/RewindDebugger/AnimNextTrace.h b/Engine/Plugins/Experimental/AnimNext/Source/AnimNext/Public/RewindDebugger/AnimNextTrace.h
new file mode 100644
--- /dev/null
+++ b/Engine/Plugins/Experimental/AnimNext/Source/AnimNext/Public/RewindDebugger/AnimNextTrace.h
@@ -0,0 +1,39 @@
+// Copyright Epic Games, Inc. All Rights Reserved.
+
+#pragma once
+
+#include "CoreTypes.h"
+#include "ObjectTrace.h"
+#include "CoreMinimal.h"
+#include "Trace/Trace.h"
+
+#define ANIMNEXT_TRACE_ENABLED (OBJECT_TRACE_ENABLED && !(UE_BUILD_SHIPPING || UE_BUILD_TEST))
+
+#if ANIMNEXT_TRACE_ENABLED
+
+UE_TRACE_CHANNEL_EXTERN(AnimNextChannel, ANIMNEXT_API);
+
+struct FAnimNextModuleInstance;
+struct FAnimNextDataInterfaceInstance;
+struct FAnimNextGraphInstance;
+
+namespace UE::AnimNext
+{
+
+struct FAnimNextTrace
+{
+	ANIMNEXT_API static void Reset();
+	
+	ANIMNEXT_API static void OutputAnimNextInstance(const FAnimNextDataInterfaceInstance* DataInterface, const UObject* OuterObject);
+	ANIMNEXT_API static void OutputAnimNextVariables(const FAnimNextDataInterfaceInstance* DataInterface, const UObject* OuterObject);
+};
+
+}
+
+
+#define TRACE_ANIMNEXT_INSTANCE(DataInterface, OuterObject) UE::AnimNext::FAnimNextTrace::OutputAnimNextInstance(DataInterface, OuterObject);
+#define TRACE_ANIMNEXT_VARIABLES(DataInterface, OuterObject) UE::AnimNext::FAnimNextTrace::OutputAnimNextVariables(DataInterface, OuterObject);
+#else
+#define TRACE_ANIMNEXT_INSTANCE(DataInterface, OuterObject)
+#define TRACE_ANIMNEXT_VARIABLES(DataInterface, OuterObject)
+#endif
diff --git a/Engine/Plugins/Experimental/AnimNext/Source/AnimNextEditor/Private/RewindDebugger/AnimNextAnalyzer.cpp b/Engine/Plugins/Experimental/AnimNext/Source/AnimNextEditor/Private/RewindDebugger/AnimNextAnalyzer.cpp
--- a/Engine/Plugins/Experimental/AnimNext/Source/AnimNextEditor/Private/RewindDebugger/AnimNextAnalyzer.cpp
+++ b/Engine/Plugins/Experimental/AnimNext/Source/AnimNextEditor/Private/RewindDebugger/AnimNextAnalyzer.cpp
@@ -20,8 +20,8 @@ void FAnimNextAnalyzer::OnAnalysisBegin(const FOnAnalysisContext& Context)
 {
 	auto& Builder = Context.InterfaceBuilder;
 
-	Builder.RouteEvent(RouteId_Module, "AnimNext", "Module");
-	Builder.RouteEvent(RouteId_ModuleVariables, "AnimNext", "ModuleVariables");
+	Builder.RouteEvent(RouteId_Module, "AnimNext", "Instance");
+	Builder.RouteEvent(RouteId_InstanceVariables, "AnimNext", "InstanceVariables");
 }
 
 bool FAnimNextAnalyzer::OnEvent(uint16 RouteId, EStyle Style, const FOnEventContext& Context)
@@ -35,15 +35,16 @@ bool FAnimNextAnalyzer::OnEvent(uint16 RouteId, EStyle Style, const FOnEventCont
 	{
 		case RouteId_Module:
 		{
-			uint64 ModuleInstanceId = EventData.GetValue<uint64>("ModuleInstanceId");
-			uint64 ModuleAssetId = EventData.GetValue<uint64>("ModuleAssetId");
+			uint64 InstanceId = EventData.GetValue<uint64>("InstanceId");
+			uint64 HostInstanceId = EventData.GetValue<uint64>("HostInstanceId");
+			uint64 AssetId = EventData.GetValue<uint64>("AssetId");
 			uint64 OuterObjectId = EventData.GetValue<uint64>("OuterObjectId");
-			Provider.AppendModule(ModuleInstanceId, ModuleAssetId, OuterObjectId);
+			Provider.AppendInstance(InstanceId, HostInstanceId, AssetId, OuterObjectId);
 			break;
 		}
-		case RouteId_ModuleVariables:
+		case RouteId_InstanceVariables:
 		{
-			uint64 ModuleInstanceId = EventData.GetValue<uint64>("ModuleInstanceId");
+			uint64 ModuleInstanceId = EventData.GetValue<uint64>("InstanceId");
 			uint64 Cycle = EventData.GetValue<uint64>("Cycle");
 			double RecordingTime = EventData.GetValue<double>("RecordingTime");
 
diff --git a/Engine/Plugins/Experimental/AnimNext/Source/AnimNextEditor/Private/RewindDebugger/AnimNextAnalyzer.h b/Engine/Plugins/Experimental/AnimNext/Source/AnimNextEditor/Private/RewindDebugger/AnimNextAnalyzer.h
--- a/Engine/Plugins/Experimental/AnimNext/Source/AnimNextEditor/Private/RewindDebugger/AnimNextAnalyzer.h
+++ b/Engine/Plugins/Experimental/AnimNext/Source/AnimNextEditor/Private/RewindDebugger/AnimNextAnalyzer.h
@@ -20,7 +20,7 @@ private:
 	enum : uint16
 	{
 		RouteId_Module,
-		RouteId_ModuleVariables,
+		RouteId_InstanceVariables,
 	};
 
 	TraceServices::IAnalysisSession& Session;
diff --git a/Engine/Plugins/Experimental/AnimNext/Source/AnimNextEditor/Private/RewindDebugger/AnimNextModuleTrack.cpp b/Engine/Plugins/Experimental/AnimNext/Source/AnimNextEditor/Private/RewindDebugger/AnimNextModuleTrack.cpp
--- a/Engine/Plugins/Experimental/AnimNext/Source/AnimNextEditor/Private/RewindDebugger/AnimNextModuleTrack.cpp
+++ b/Engine/Plugins/Experimental/AnimNext/Source/AnimNextEditor/Private/RewindDebugger/AnimNextModuleTrack.cpp
@@ -32,10 +32,9 @@ FText FAnimNextModuleTrack::GetDisplayNameInternal() const
 	const FAnimNextProvider* AnimNextProvider = IRewindDebugger::Instance()->GetAnalysisSession()->ReadProvider<FAnimNextProvider>("AnimNextProvider");
 	const IGameplayProvider* GameplayProvider = IRewindDebugger::Instance()->GetAnalysisSession()->ReadProvider<IGameplayProvider>("GameplayProvider");
 
- 	uint64 ModuleId;
- 	if (AnimNextProvider->GetModuleId(ObjectId, ModuleId))
+ 	if (InstanceId != 0)
  	{
- 		if (const FDataInterfaceData* Data = AnimNextProvider->GetDataInterfaceData(ModuleId))
+ 		if (const FDataInterfaceData* Data = AnimNextProvider->GetDataInterfaceData(InstanceId))
  		{
 			const FObjectInfo& ModuleInfo = GameplayProvider->GetObjectInfo(Data->AssetId);
 			return FText::FromString(ModuleInfo.Name);
@@ -59,12 +58,14 @@ TSharedPtr<RewindDebugger::FRewindDebuggerTrack> FAnimNextModuleTrackCreator::Cr
 FAnimNextModuleTrack::FAnimNextModuleTrack(uint64 InObjectId) :
 	ObjectId(InObjectId)
 {
-	FPropertyEditorModule& PropertyEditorModule = FModuleManager::GetModuleChecked<FPropertyEditorModule>("PropertyEditor");
-	FDetailsViewArgs DetailsViewArgs;
-	DetailsViewArgs.NameAreaSettings = FDetailsViewArgs::HideNameArea;
-	DetailsView = PropertyEditorModule.CreateDetailView(DetailsViewArgs);
-
-	InitializeDetailsObject();
+	Initialize();
+}
+	
+FAnimNextModuleTrack::FAnimNextModuleTrack(uint64 InObjectId, uint64 InInstanceId) :
+   	ObjectId(InObjectId),
+	InstanceId(InInstanceId)
+{
+   	Initialize();
 }
 	
 	
@@ -76,52 +77,85 @@ FAnimNextModuleTrack::~FAnimNextModuleTrack()
    	}
 }
 	
+void FAnimNextModuleTrack::Initialize()
+{
+	ExistenceRange = MakeShared<SEventTimelineView::FTimelineEventData>();
+	ExistenceRange->Windows.Add({0,0, GetDisplayNameInternal(), GetDisplayNameInternal(), FLinearColor(0.1f,0.15f,0.11f)});
+	
+   	FPropertyEditorModule& PropertyEditorModule = FModuleManager::GetModuleChecked<FPropertyEditorModule>("PropertyEditor");
+   	FDetailsViewArgs DetailsViewArgs;
+   	DetailsViewArgs.NameAreaSettings = FDetailsViewArgs::HideNameArea;
+   	DetailsView = PropertyEditorModule.CreateDetailView(DetailsViewArgs);
+
+	InitializeDetailsObject();
+}
+	
 UPropertyBagDetailsObject* FAnimNextModuleTrack::InitializeDetailsObject()
 {
 	UPropertyBagDetailsObject* DetailsObject = NewObject<UPropertyBagDetailsObject>();
 	DetailsObject->SetFlags(RF_Standalone);
 	DetailsObjectWeakPtr = MakeWeakObjectPtr(DetailsObject);
 	DetailsView->SetObject(DetailsObject);
-	return DetailsObject;
+ 	return DetailsObject;
 }
 
 
+TSharedPtr<SWidget> FAnimNextModuleTrack::GetTimelineViewInternal()
+{
+	return SNew(SEventTimelineView)
+		.ViewRange_Lambda([]() { return IRewindDebugger::Instance()->GetCurrentViewRange(); })
+		.EventData_Raw(this, &FAnimNextModuleTrack::GetExistenceRange);}
+
 bool FAnimNextModuleTrack::UpdateInternal()
 {
 	TRACE_CPUPROFILER_EVENT_SCOPE(FAnimNextModuleTrack::UpdateInternal);
  	
 	IRewindDebugger* RewindDebugger = IRewindDebugger::Instance();
+ 	const TRange<double> ViewRange = RewindDebugger->GetCurrentViewRange();
 	
 	const TraceServices::IAnalysisSession* AnalysisSession = RewindDebugger->GetAnalysisSession();
+
+ 	bool bChanged = false;
 	
 	if (const FAnimNextProvider* AnimNextProvider = IRewindDebugger::Instance()->GetAnalysisSession()->ReadProvider<FAnimNextProvider>("AnimNextProvider"))
 	{
 		double CurrentScrubTime = IRewindDebugger::Instance()->CurrentTraceTime();
-		if (PreviousScrubTime != CurrentScrubTime)
-		{
-			PreviousScrubTime = CurrentScrubTime;
 
-			UPropertyBagDetailsObject* DetailsObject = DetailsObjectWeakPtr.Get();
-			if (DetailsObject == nullptr)
-			{
-				// this should not happen unless the object was garbage collected (which should not happen since it's marked as Standalone)
-				DetailsObject = InitializeDetailsObject();
-			}
+		UPropertyBagDetailsObject* DetailsObject = DetailsObjectWeakPtr.Get();
+		if (DetailsObject == nullptr)
+		{
+			// this should not happen unless the object was garbage collected (which should not happen since it's marked as Standalone)
+			Initialize();
+		}
 
+		if (InstanceId == 0)
+		{
+			AnimNextProvider->GetModuleId(ObjectId, InstanceId);
+		}
+		
+		if (RewindDebugger->IsTraceFileLoaded()) // this currently crashes for object references on loaded trace files, so temporarily disabling it
+		{
 			DetailsObject->Properties.Reset();
+		}
 
-			if (!RewindDebugger->IsTraceFileLoaded()) // this currently crashes for object references on loaded trace files, so temporarily disabling it
-			{
-				const TraceServices::IFrameProvider& FramesProvider = TraceServices::ReadFrameProvider(*AnalysisSession);
-				TraceServices::FAnalysisSessionReadScope SessionReadScope(*AnalysisSession);
+		if(InstanceId != 0)
+		{
 			
-				TraceServices::FFrame MarkerFrame;
-				if(FramesProvider.GetFrameFromTime(ETraceFrameType::TraceFrameType_Game, CurrentScrubTime, MarkerFrame))
+			if (const FDataInterfaceData* Data = AnimNextProvider->GetDataInterfaceData(InstanceId))
+			{
+				if (PreviousScrubTime != CurrentScrubTime)
 				{
-					uint64 ModuleId;
-					if (AnimNextProvider->GetModuleId(ObjectId, ModuleId))
+					PreviousScrubTime = CurrentScrubTime;
+					
+					if (!RewindDebugger->IsTraceFileLoaded()) // this currently crashes for object references on loaded trace files, so temporarily disabling it
 					{
-						if (const FDataInterfaceData* Data = AnimNextProvider->GetDataInterfaceData(ModuleId))
+						const TraceServices::IFrameProvider& FramesProvider = TraceServices::ReadFrameProvider(*AnalysisSession);
+						TraceServices::FAnalysisSessionReadScope SessionReadScope(*AnalysisSession);
+
+						DetailsObject->Properties.Reset();
+
+						TraceServices::FFrame MarkerFrame;
+						if(FramesProvider.GetFrameFromTime(ETraceFrameType::TraceFrameType_Game, CurrentScrubTime, MarkerFrame))
 						{
 							Data->VariablesTimeline.EnumerateEvents(MarkerFrame.StartTime, MarkerFrame.EndTime, [DetailsObject](double InStartTime, double InEndTime, uint32 InDepth, const TArray<uint8>& VariableData)
 							{
@@ -133,15 +167,50 @@ bool FAnimNextModuleTrack::UpdateInternal()
 						}
 					}
 				}
+		
+				ExistenceRange->Windows.SetNum(1,EAllowShrinking::No);
+				ExistenceRange->Windows[0].TimeStart = Data->StartTime;
+				ExistenceRange->Windows[0].TimeEnd = Data->EndTime;
 			}
+
+			TArray<uint64, TInlineAllocator<32>> CurrentChildren;
+		
+			// update/create child tracks
+			AnimNextProvider->EnumerateChildInstances(InstanceId, [this, &bChanged, &CurrentChildren, &ViewRange](const FDataInterfaceData& ChildData)
+			{
+				if (ChildData.StartTime < ViewRange.GetUpperBoundValue() && ChildData.EndTime > ViewRange.GetLowerBoundValue())
+				{
+					CurrentChildren.Add(ChildData.InstanceId);
+					if (!Children.ContainsByPredicate([&ChildData](TSharedPtr<FAnimNextModuleTrack>& Child) { return Child->InstanceId == ChildData.InstanceId; } ))
+					{
+						Children.Add(MakeShared<FAnimNextModuleTrack>(ObjectId, ChildData.InstanceId));
+						bChanged = true;
+					}
+				}
+			});
+			
+			int32 NumRemoved = Children.RemoveAll([&CurrentChildren](TSharedPtr<FAnimNextModuleTrack>& Child)
+			{
+				return !CurrentChildren.Contains(Child->InstanceId);
+			});
+			bChanged |= (NumRemoved > 0);
+		}
+
+		for (auto& Child : Children)
+		{
+			bChanged |= Child->Update();
 		}
 	}
 	
- 	return false;
+ 	return bChanged;
 }
 
 void FAnimNextModuleTrack::IterateSubTracksInternal(TFunction<void(TSharedPtr<FRewindDebuggerTrack> SubTrack)> IteratorFunction)
 {
+	for(TSharedPtr<FAnimNextModuleTrack>& Track : Children)
+	{
+		IteratorFunction(Track);
+	}
 };
 
 bool FAnimNextModuleTrack::HandleDoubleClickInternal()
@@ -154,10 +223,9 @@ bool FAnimNextModuleTrack::HandleDoubleClickInternal()
 		const IGameplayProvider* GameplayProvider = AnalysisSession->ReadProvider<IGameplayProvider>("GameplayProvider");
 		const FAnimNextProvider* AnimNextProvider = IRewindDebugger::Instance()->GetAnalysisSession()->ReadProvider<FAnimNextProvider>("AnimNextProvider");
 
-		uint64 ModuleId;
-		if (AnimNextProvider->GetModuleId(ObjectId, ModuleId))
+		if (InstanceId != 0)
 		{
-			if (const FDataInterfaceData* ModuleData = AnimNextProvider->GetDataInterfaceData(ModuleId))
+			if (const FDataInterfaceData* ModuleData = AnimNextProvider->GetDataInterfaceData(InstanceId))
 			{
 				const FObjectInfo& AssetInfo = GameplayProvider->GetObjectInfo(ModuleData->AssetId);
 				GEditor->GetEditorSubsystem<UAssetEditorSubsystem>()->OpenEditorForAsset(AssetInfo.PathName);
diff --git a/Engine/Plugins/Experimental/AnimNext/Source/AnimNextEditor/Private/RewindDebugger/AnimNextModuleTrack.h b/Engine/Plugins/Experimental/AnimNext/Source/AnimNextEditor/Private/RewindDebugger/AnimNextModuleTrack.h
--- a/Engine/Plugins/Experimental/AnimNext/Source/AnimNextEditor/Private/RewindDebugger/AnimNextModuleTrack.h
+++ b/Engine/Plugins/Experimental/AnimNext/Source/AnimNextEditor/Private/RewindDebugger/AnimNextModuleTrack.h
@@ -7,6 +7,7 @@
 #include "SEventTimelineView.h"
 #include "StructUtils/PropertyBag.h"
 #include "IDetailsView.h"
+#include "SEventTimelineView.h"
 
 #include "AnimNextModuleTrack.generated.h"
 
@@ -27,8 +28,12 @@ class FAnimNextModuleTrack : public RewindDebugger::FRewindDebuggerTrack
 {
 public:
 	FAnimNextModuleTrack(uint64 InObjectId);
+	FAnimNextModuleTrack(uint64 InObjectId, uint64 InstanceId);
 	virtual ~FAnimNextModuleTrack();
+	
+	TSharedPtr<SEventTimelineView::FTimelineEventData> GetExistenceRange() const { return ExistenceRange; }
 private:
+	virtual TSharedPtr<SWidget> GetTimelineViewInternal() override;
 	virtual bool UpdateInternal() override;
 	virtual TSharedPtr<SWidget> GetDetailsViewInternal() override { return DetailsView; }
 
@@ -39,13 +44,18 @@ private:
 	virtual void IterateSubTracksInternal(TFunction<void(TSharedPtr<FRewindDebuggerTrack> SubTrack)> IteratorFunction) override;
 	virtual bool HandleDoubleClickInternal() override;
 
+	void Initialize();
 	UPropertyBagDetailsObject* InitializeDetailsObject();
-	
+
 	TSharedPtr<IDetailsView> DetailsView;
 	FSlateIcon Icon;
 	uint64 ObjectId;
+	uint64 InstanceId = 0; 
 	double PreviousScrubTime = -1.0;
 	TWeakObjectPtr<UPropertyBagDetailsObject> DetailsObjectWeakPtr;
+	TSharedPtr<SEventTimelineView::FTimelineEventData> ExistenceRange;
+
+	TArray<TSharedPtr<FAnimNextModuleTrack>> Children;
 };
 
 
diff --git a/Engine/Plugins/Experimental/AnimNext/Source/AnimNextEditor/Private/RewindDebugger/AnimNextProvider.cpp b/Engine/Plugins/Experimental/AnimNext/Source/AnimNextEditor/Private/RewindDebugger/AnimNextProvider.cpp
--- a/Engine/Plugins/Experimental/AnimNext/Source/AnimNextEditor/Private/RewindDebugger/AnimNextProvider.cpp
+++ b/Engine/Plugins/Experimental/AnimNext/Source/AnimNextEditor/Private/RewindDebugger/AnimNextProvider.cpp
@@ -12,27 +12,36 @@ FAnimNextProvider::FAnimNextProvider(TraceServices::IAnalysisSession& InSession)
 {
 }
 
-void FAnimNextProvider::AppendModule(uint64 InModuleInstanceId, uint64 InModuleAssetId, uint64 InOuterObjectId)
+void FAnimNextProvider::AppendInstance(uint64 InstanceId, uint64 HostInstanceId, uint64 AssetId, uint64 OuterObjectId)
 {
 	Session.WriteAccessCheck();
 
-	ComponentIdToModuleId.Add(InOuterObjectId, InModuleInstanceId);
+	if (HostInstanceId == 0)
+	{
+		ComponentIdToModuleId.Add(OuterObjectId, InstanceId);
+	}
 
 	TSharedRef<FDataInterfaceData> Data = MakeShared<FDataInterfaceData>(Session);
-	Data->InstanceId = InModuleInstanceId;
-	Data->AssetId = InModuleAssetId;
+	Data->InstanceId = InstanceId;
+	Data->HostInstanceId = HostInstanceId;
+	Data->OuterObjectId = OuterObjectId;
+	Data->AssetId = AssetId;
 
-	DataInterfaceData.Add( InModuleInstanceId, Data);
+	DataInterfaceData.Add(InstanceId, Data);
 }
 
 void FAnimNextProvider::AppendVariables(double ProfileTime, double RecordingTime, uint64 DataInterfaceId, const TArrayView<const uint8>& VariableData)
 {
-
 	Session.WriteAccessCheck();
 	
 	if (const TSharedRef<FDataInterfaceData>* Data = DataInterfaceData.Find(DataInterfaceId))
 	{
 		(*Data)->VariablesTimeline.AppendEvent(ProfileTime, TArray<uint8>(VariableData));
+		(*Data)->EndTime = RecordingTime;
+		if ((*Data)->StartTime < 0)
+		{
+			(*Data)->StartTime = RecordingTime;
+		}
 	}
 }
 
@@ -58,4 +67,19 @@ const FDataInterfaceData* FAnimNextProvider::GetDataInterfaceData(uint64 DataInt
 	return nullptr;
 }
 
+void FAnimNextProvider::EnumerateChildInstances(uint64 InstanceId, TFunctionRef<void(const FDataInterfaceData&)> Callback) const
+{
+	Session.ReadAccessCheck();
+
+	for (auto& Data : DataInterfaceData)
+	{
+		const TSharedPtr<FDataInterfaceData>& Instance = Data.Value;
+		
+		if (Instance->HostInstanceId == InstanceId)
+		{
+			Callback(*Instance);
+		}
+	}
+}
+
 #undef LOCTEXT_NAMESPACE
diff --git a/Engine/Plugins/Experimental/AnimNext/Source/AnimNextEditor/Private/RewindDebugger/AnimNextProvider.h b/Engine/Plugins/Experimental/AnimNext/Source/AnimNextEditor/Private/RewindDebugger/AnimNextProvider.h
--- a/Engine/Plugins/Experimental/AnimNext/Source/AnimNextEditor/Private/RewindDebugger/AnimNextProvider.h
+++ b/Engine/Plugins/Experimental/AnimNext/Source/AnimNextEditor/Private/RewindDebugger/AnimNextProvider.h
@@ -16,13 +16,21 @@ struct FDataInterfaceData
 	FDataInterfaceData(TraceServices::IAnalysisSession& Session)
 		: VariablesTimeline(Session.GetLinearAllocator( )),
 			InstanceId (0),
-			AssetId (0)
+			HostInstanceId (0),
+			OuterObjectId (0),
+			AssetId (0),
+			StartTime (-1),
+			EndTime (-1)
 	{
 	}
 
 	TraceServices::TPointTimeline<TArray<uint8>> VariablesTimeline;
 	uint64 InstanceId;
+	uint64 HostInstanceId;
+	uint64 OuterObjectId;
 	uint64 AssetId;
+	double StartTime;
+	double EndTime;
 };
 
 class FAnimNextProvider : public TraceServices::IProvider
@@ -32,12 +40,14 @@ public:
 
 	FAnimNextProvider(TraceServices::IAnalysisSession& InSession);
 
-	void AppendModule(uint64 InModuleInstanceId, uint64 InModuleAssetId, uint64 InOuterObjectId);
+	void AppendInstance(uint64 InstanceId, uint64 HostInstanceId,  uint64 AssetId, uint64 OuterObjectId);
 	void AppendVariables(double ProfileTime, double RecordingTime, uint64 InModuleInstanceId, const TArrayView<const uint8>& VariableData);
 
 	bool GetModuleId(uint64 ComponentId, uint64& OutModuleId) const;
 	const FDataInterfaceData* GetDataInterfaceData(uint64 DataInterfaceId) const;
 
+	void EnumerateChildInstances(uint64 InstanceId, TFunctionRef<void(const FDataInterfaceData&)> Callback) const;
+
 private:
 	TraceServices::IAnalysisSession& Session;
 
diff --git a/Engine/Plugins/Experimental/AnimNextAnimGraph/Source/AnimNextAnimGraph/Private/Graph/RigUnit_AnimNextRunAnimationGraph_v2.cpp b/Engine/Plugins/Experimental/AnimNextAnimGraph/Source/AnimNextAnimGraph/Private/Graph/RigUnit_AnimNextRunAnimationGraph_v2.cpp
--- a/Engine/Plugins/Experimental/AnimNextAnimGraph/Source/AnimNextAnimGraph/Private/Graph/RigUnit_AnimNextRunAnimationGraph_v2.cpp
+++ b/Engine/Plugins/Experimental/AnimNextAnimGraph/Source/AnimNextAnimGraph/Private/Graph/RigUnit_AnimNextRunAnimationGraph_v2.cpp
@@ -15,6 +15,7 @@
 #include "Injection/InjectionDataInterfaceHostAdapter.h"
 #include "TraitCore/TraitEventList.h"
 #include "Injection/AnimNextModuleInjectionComponent.h"
+#include "TraceAnimNextGraphInstances.h"
 
 #include UE_INLINE_GENERATED_CPP_BY_NAME(RigUnit_AnimNextRunAnimationGraph_v2)
 
@@ -51,6 +52,7 @@ FRigUnit_AnimNextRunAnimationGraph_v2_Execute()
 		return;
 	}
 
+
 	// Release the instance if the host graph has changed
 	if(WorkData.WeakHost.IsValid() && !WorkData.WeakHost.Pin()->UsesAnimationGraph(HostGraph))
 	{
@@ -141,6 +143,9 @@ FRigUnit_AnimNextRunAnimationGraph_v2_Execute()
 
 	UpdateGraph(HostInstance, ExecuteContext.GetDeltaTime(), InputEventList, OutputEventList);
 	EvaluateGraph(HostInstance, RefPose, LOD, Result);
+	
+	TRACE_ANIMNEXT_MODULE(ModuleInstance);
+	TRACE_ANIMNEXT_GRAPHINSTANCES(HostInstance);
 
 	// We might have appended new input/output events, append them
 	{
diff --git a/Engine/Plugins/Experimental/AnimNextAnimGraph/Source/AnimNextAnimGraph/Private/Graph/TraceAnimNextGraphInstances.cpp b/Engine/Plugins/Experimental/AnimNextAnimGraph/Source/AnimNextAnimGraph/Private/Graph/TraceAnimNextGraphInstances.cpp
new file mode 100644
--- /dev/null
+++ b/Engine/Plugins/Experimental/AnimNextAnimGraph/Source/AnimNextAnimGraph/Private/Graph/TraceAnimNextGraphInstances.cpp
@@ -0,0 +1,58 @@
+// Copyright Epic Games, Inc. All Rights Reserved.
+
+#include "TraceAnimNextGraphInstances.h"
+#include "Graph/AnimNextGraphInstance.h"
+#include "Module/AnimNextModuleInstance.h"
+#include "TraitCore/ExecutionContext.h"
+#include "TraitCore/NodeInstance.h"
+#include "TraitInterfaces/IHierarchy.h"
+
+#if ANIMNEXT_TRACE_ENABLED
+namespace UE::AnimNext
+{
+	void TraceGraphInstances(const FAnimNextGraphInstance& RootGraph)
+	{
+		bool bChannelEnabled = UE_TRACE_CHANNELEXPR_IS_ENABLED(AnimNextChannel);
+    	if (!bChannelEnabled)
+    	{
+    		return;
+    	}
+
+		UObject* OuterObject = RootGraph.GetModuleInstance() ? RootGraph.GetModuleInstance()->GetObject() : nullptr;
+
+		FExecutionContext ExecutionContext;
+		TArray<FWeakTraitPtr, TInlineAllocator<8, TMemStackAllocator<>>> Traits;
+		TSet<uint64> TracedInstances;
+		
+		FMemStack& MemStack = FMemStack::Get();
+		FMemMark Mark(MemStack);
+		
+		ExecutionContext.BindTo(RootGraph.GetGraphRootPtr());
+
+		Traits.Push(RootGraph.GetGraphRootPtr());
+
+		TTraitBinding<IHierarchy> HierarchyTrait;
+
+		while (!Traits.IsEmpty())
+		{
+			FWeakTraitPtr Trait = Traits.Pop();
+			if (Trait.IsValid())
+			{
+				FAnimNextGraphInstance& Graph = Trait.GetNodeInstance()->GetOwner();
+				if (!TracedInstances.Contains(Graph.GetUniqueId()))
+				{
+					TRACE_ANIMNEXT_VARIABLES(&Graph, Graph.GetModuleInstance() ? Graph.GetModuleInstance()->GetObject() : nullptr);
+				}
+
+				FTraitStackBinding TraitStack;
+				ExecutionContext.GetStack(Trait, TraitStack);
+				
+				if (TraitStack.GetInterface(HierarchyTrait))
+				{
+					IHierarchy::GetStackChildren(ExecutionContext, TraitStack, Traits);
+				}
+			}
+		}
+	}
+}
+#endif
diff --git a/Engine/Plugins/Experimental/AnimNextAnimGraph/Source/AnimNextAnimGraph/Private/Graph/TraceAnimNextGraphInstances.h b/Engine/Plugins/Experimental/AnimNextAnimGraph/Source/AnimNextAnimGraph/Private/Graph/TraceAnimNextGraphInstances.h
new file mode 100644
--- /dev/null
+++ b/Engine/Plugins/Experimental/AnimNextAnimGraph/Source/AnimNextAnimGraph/Private/Graph/TraceAnimNextGraphInstances.h
@@ -0,0 +1,21 @@
+// Copyright Epic Games, Inc. All Rights Reserved.
+#pragma once
+
+#include "RewindDebugger/AnimNextTrace.h"
+
+#if ANIMNEXT_TRACE_ENABLED
+
+struct FAnimNextGraphInstance;
+
+namespace UE::AnimNext
+{
+	void TraceGraphInstances(const FAnimNextGraphInstance& RootGraph);
+}
+
+#define TRACE_ANIMNEXT_GRAPHINSTANCES(RootGraph) UE::AnimNext::TraceGraphInstances(RootGraph);
+
+#else
+
+#define TRACE_ANIMNEXT_GRAPHINSTANCES(RootGraph)
+
+#endif
