Fix crash tracing Anim Next graph instances with no variables

diff --git a/Engine/Plugins/Experimental/AnimNext/Source/AnimNext/Private/RewindDebugger/AnimNextTrace.cpp b/Engine/Plugins/Experimental/AnimNext/Source/AnimNext/Private/RewindDebugger/AnimNextTrace.cpp
--- a/Engine/Plugins/Experimental/AnimNext/Source/AnimNext/Private/RewindDebugger/AnimNextTrace.cpp
+++ b/Engine/Plugins/Experimental/AnimNext/Source/AnimNext/Private/RewindDebugger/AnimNextTrace.cpp
@@ -97,28 +97,38 @@ void FAnimNextTrace::OutputAnimNextVariables(const FAnimNextDataInterfaceInstanc
 
 	if (CANNOT_TRACE_OBJECT(OuterObject))
 	{
-      	return;
+		return;
 	}
 
 	OutputAnimNextInstance(Instance, OuterObject);
 	
 	FInstancedPropertyBag Variables = Instance->GetVariables();
-	TConstArrayView<FPropertyBagPropertyDesc> PropertyDescs = Variables.GetPropertyBagStruct()->GetPropertyDescs();
-	for(int i =0; i<PropertyDescs.Num(); i++)
+	if (Variables.GetPropertyBagStruct())
 	{
-		uint8* ValueMemory = Instance->GetMemoryForVariable(i, PropertyDescs[i].Name, PropertyDescs[i].CachedProperty);
-		PropertyDescs[i].CachedProperty->CopySingleValue(const_cast<uint8*>(Variables.GetValue().GetMemory()) +  PropertyDescs[i].CachedProperty->GetOffset_ForInternal(), ValueMemory);
-	}
+		TConstArrayView<FPropertyBagPropertyDesc> PropertyDescs = Variables.GetPropertyBagStruct()->GetPropertyDescs();
+		for(int i =0; i<PropertyDescs.Num(); i++)
+		{
+			uint8* ValueMemory = Instance->GetMemoryForVariable(i, PropertyDescs[i].Name, PropertyDescs[i].CachedProperty);
+			PropertyDescs[i].CachedProperty->CopySingleValue(const_cast<uint8*>(Variables.GetValue().GetMemory()) +  PropertyDescs[i].CachedProperty->GetOffset_ForInternal(), ValueMemory);
+		}
 
-	TArray<uint8> ArchiveData;
-	FObjectWriter Archive(ArchiveData);
-	const_cast<FInstancedPropertyBag&>(Variables).Serialize(Archive);
+		TArray<uint8> ArchiveData;
+		FObjectWriter Archive(ArchiveData);
+		const_cast<FInstancedPropertyBag&>(Variables).Serialize(Archive);
 
-	UE_TRACE_LOG(AnimNext, InstanceVariables, AnimNextChannel)
-			<< InstanceVariables.Cycle(FPlatformTime::Cycles64())
-			<< InstanceVariables.RecordingTime(FObjectTrace::GetWorldElapsedTime(OuterObject->GetWorld()))
-			<< InstanceVariables.InstanceId(InstanceId)
-			<< InstanceVariables.VariableData(ArchiveData.GetData(), ArchiveData.Num());
+		UE_TRACE_LOG(AnimNext, InstanceVariables, AnimNextChannel)
+				<< InstanceVariables.Cycle(FPlatformTime::Cycles64())
+				<< InstanceVariables.RecordingTime(FObjectTrace::GetWorldElapsedTime(OuterObject->GetWorld()))
+				<< InstanceVariables.InstanceId(InstanceId)
+				<< InstanceVariables.VariableData(ArchiveData.GetData(), ArchiveData.Num());
+	}
+	else
+	{
+		UE_TRACE_LOG(AnimNext, InstanceVariables, AnimNextChannel)
+				<< InstanceVariables.Cycle(FPlatformTime::Cycles64())
+				<< InstanceVariables.RecordingTime(FObjectTrace::GetWorldElapsedTime(OuterObject->GetWorld()))
+				<< InstanceVariables.InstanceId(InstanceId);
+	}
 }
 	
 }
diff --git a/Engine/Plugins/Experimental/AnimNext/Source/AnimNextEditor/Private/RewindDebugger/AnimNextProvider.cpp b/Engine/Plugins/Experimental/AnimNext/Source/AnimNextEditor/Private/RewindDebugger/AnimNextProvider.cpp
--- a/Engine/Plugins/Experimental/AnimNext/Source/AnimNextEditor/Private/RewindDebugger/AnimNextProvider.cpp
+++ b/Engine/Plugins/Experimental/AnimNext/Source/AnimNextEditor/Private/RewindDebugger/AnimNextProvider.cpp
@@ -36,7 +36,10 @@ void FAnimNextProvider::AppendVariables(double ProfileTime, double RecordingTime
 	
 	if (const TSharedRef<FDataInterfaceData>* Data = DataInterfaceData.Find(DataInterfaceId))
 	{
-		(*Data)->VariablesTimeline.AppendEvent(ProfileTime, TArray<uint8>(VariableData));
+		if (!VariableData.IsEmpty())
+		{
+			(*Data)->VariablesTimeline.AppendEvent(ProfileTime, TArray<uint8>(VariableData));
+		}
 		(*Data)->EndTime = RecordingTime;
 		if ((*Data)->StartTime < 0)
 		{
