Fix Anim Next Variable display in Rewind Debugger for traces loaded from disk

diff --git a/Engine/Plugins/Experimental/AnimNext/Source/AnimNext/Private/RewindDebugger/AnimNextTrace.cpp b/Engine/Plugins/Experimental/AnimNext/Source/AnimNext/Private/RewindDebugger/AnimNextTrace.cpp
--- a/Engine/Plugins/Experimental/AnimNext/Source/AnimNext/Private/RewindDebugger/AnimNextTrace.cpp
+++ b/Engine/Plugins/Experimental/AnimNext/Source/AnimNext/Private/RewindDebugger/AnimNextTrace.cpp
@@ -6,7 +6,8 @@
 #include "Chaos/ParticleDirtyFlags.h"
 #include "DataInterface/AnimNextDataInterfaceInstance.h"
 #include "Module/AnimNextModuleInstance.h"
-#include "Serialization/ObjectWriter.h"
+#include "Serialization/MemoryWriter.h"
+#include "Serialization/ObjectAndNameAsStringProxyArchive.h"
 
 #if ANIMNEXT_TRACE_ENABLED
 
@@ -113,7 +114,8 @@ void FAnimNextTrace::OutputAnimNextVariables(const FAnimNextDataInterfaceInstanc
 		}
 
 		TArray<uint8> ArchiveData;
-		FObjectWriter Archive(ArchiveData);
+		FMemoryWriter WriterArchive(ArchiveData);
+		FObjectAndNameAsStringProxyArchive Archive(WriterArchive, /*bInLoadIfFindFails*/true);
 		const_cast<FInstancedPropertyBag&>(Variables).Serialize(Archive);
 
 		UE_TRACE_LOG(AnimNext, InstanceVariables, AnimNextChannel)
diff --git a/Engine/Plugins/Experimental/AnimNext/Source/AnimNextEditor/Private/RewindDebugger/AnimNextModuleTrack.cpp b/Engine/Plugins/Experimental/AnimNext/Source/AnimNextEditor/Private/RewindDebugger/AnimNextModuleTrack.cpp
--- a/Engine/Plugins/Experimental/AnimNext/Source/AnimNextEditor/Private/RewindDebugger/AnimNextModuleTrack.cpp
+++ b/Engine/Plugins/Experimental/AnimNext/Source/AnimNextEditor/Private/RewindDebugger/AnimNextModuleTrack.cpp
@@ -9,6 +9,7 @@
 #include "Subsystems/AssetEditorSubsystem.h"
 #include "Modules/ModuleManager.h"
 #include "PropertyEditorModule.h"
+#include "Serialization/ObjectAndNameAsStringProxyArchive.h"
 #include "Serialization/ObjectReader.h"
 
 
@@ -133,38 +134,30 @@ bool FAnimNextModuleTrack::UpdateInternal()
 			AnimNextProvider->GetModuleId(ObjectId, InstanceId);
 		}
 		
-		if (RewindDebugger->IsTraceFileLoaded()) // this currently crashes for object references on loaded trace files, so temporarily disabling it
-		{
-			DetailsObject->Properties.Reset();
-		}
-
 		if(InstanceId != 0)
 		{
-			
 			if (const FDataInterfaceData* Data = AnimNextProvider->GetDataInterfaceData(InstanceId))
 			{
 				if (PreviousScrubTime != CurrentScrubTime)
 				{
 					PreviousScrubTime = CurrentScrubTime;
 					
-					if (!RewindDebugger->IsTraceFileLoaded()) // this currently crashes for object references on loaded trace files, so temporarily disabling it
+					const TraceServices::IFrameProvider& FramesProvider = TraceServices::ReadFrameProvider(*AnalysisSession);
+					TraceServices::FAnalysisSessionReadScope SessionReadScope(*AnalysisSession);
+
+					DetailsObject->Properties.Reset();
+
+					TraceServices::FFrame MarkerFrame;
+					if(FramesProvider.GetFrameFromTime(ETraceFrameType::TraceFrameType_Game, CurrentScrubTime, MarkerFrame))
 					{
-						const TraceServices::IFrameProvider& FramesProvider = TraceServices::ReadFrameProvider(*AnalysisSession);
-						TraceServices::FAnalysisSessionReadScope SessionReadScope(*AnalysisSession);
-
-						DetailsObject->Properties.Reset();
-
-						TraceServices::FFrame MarkerFrame;
-						if(FramesProvider.GetFrameFromTime(ETraceFrameType::TraceFrameType_Game, CurrentScrubTime, MarkerFrame))
-						{
-							Data->VariablesTimeline.EnumerateEvents(MarkerFrame.StartTime, MarkerFrame.EndTime, [DetailsObject](double InStartTime, double InEndTime, uint32 InDepth, const TArray<uint8>& VariableData)
-							{
-								FObjectReader Archive(VariableData);
-								DetailsObject->Properties.Serialize(Archive);
-							
-								return TraceServices::EEventEnumerate::Stop;
-							});
-						}
+						Data->VariablesTimeline.EnumerateEvents(MarkerFrame.StartTime, MarkerFrame.EndTime, [DetailsObject](double InStartTime, double InEndTime, uint32 InDepth, const TArray<uint8>& VariableData)
+						{				
+							FMemoryReader Reader(VariableData);
+							FObjectAndNameAsStringProxyArchive ReaderProxy(Reader, /*bInLoadIfFindFails*/true);
+							DetailsObject->Properties.Serialize(ReaderProxy);
+						
+							return TraceServices::EEventEnumerate::Stop;
+						});
 					}
 				}
 		
