Crash in remap pose node when either the source or target mesh are invalid when remapping attributes

diff --git a/Engine/Plugins/Experimental/AnimNext/Source/AnimNext/Private/Graph/RemapPoseData.cpp b/Engine/Plugins/Experimental/AnimNext/Source/AnimNext/Private/Graph/RemapPoseData.cpp
--- a/Engine/Plugins/Experimental/AnimNext/Source/AnimNext/Private/Graph/RemapPoseData.cpp
+++ b/Engine/Plugins/Experimental/AnimNext/Source/AnimNext/Private/Graph/RemapPoseData.cpp
@@ -173,10 +173,20 @@ void FRemapPoseData::RemapAttributes(const UE::AnimNext::FLODPose& SourceLODPose
 	UE::Anim::FHeapAttributeContainer& OutAttributes)
 {
 	const USkeletalMesh* SourceSkeletalMesh = SourceRefPose->SkeletalMesh.Get();
+	if (!SourceSkeletalMesh)
+	{
+		ensureMsgf(false, TEXT("FRemapPoseData::RemapAttributes(): Source skeletal mesh is not valid anymore."));
+		return;
+	}
 	const FReferenceSkeleton& SourceRefSkeleton = SourceSkeletalMesh->GetRefSkeleton();
 	const TArrayView<const FBoneIndexType> SourceLODBoneIndexToSkeletonBoneIndexMap = SourceLODPose.GetLODBoneIndexToSkeletonBoneIndexMap();
 
 	const USkeletalMesh* TargetSkeletalMesh = TargetRefPose->SkeletalMesh.Get();
+	if (!TargetSkeletalMesh)
+	{
+		ensureMsgf(false, TEXT("FRemapPoseData::RemapAttributes(): Target skeletal mesh is not valid anymore."));
+		return;
+	}
 	const FReferenceSkeleton& TargetRefSkeleton = TargetSkeletalMesh->GetRefSkeleton();
 	const TArrayView<const FBoneIndexType> TargetSkeletonBoneToLODBoneIndexMap = TargetLODPose.GetSkeletonBoneIndexToLODBoneIndexMap();
 
