Enqueue adding/remove of prerequisites for processing at end of frame

diff --git a/Engine/Plugins/Experimental/AnimNext/Source/AnimNext/Internal/Component/AnimNextWorldSubsystem.h b/Engine/Plugins/Experimental/AnimNext/Source/AnimNext/Internal/Component/AnimNextWorldSubsystem.h
--- a/Engine/Plugins/Experimental/AnimNext/Source/AnimNext/Internal/Component/AnimNextWorldSubsystem.h
+++ b/Engine/Plugins/Experimental/AnimNext/Source/AnimNext/Internal/Component/AnimNextWorldSubsystem.h
@@ -33,6 +33,8 @@ struct FModulePendingAction
 		ReleaseHandle,			// Payload = FModulePayloadNone
 		EnableHandle,			// Payload = bool
 		EnableDebugDrawing,		// Payload = bool
+		AddPrerequisite,		// Payload = UE::AnimNext::FModuleHandle (to-be-added handle)
+		RemovePrerequisite,	    // Payload = UE::AnimNext::FModuleHandle (to-be-removed handle)
 	};
 
 	// Marker struct for an 'empty' payload
@@ -57,7 +59,7 @@ struct FModulePendingAction
 		Payload.Set<PayloadType>(InPayload);
 	}
 
-	TVariant<FModulePayloadNone, bool> Payload;
+	TVariant<FModulePayloadNone, bool, UE::AnimNext::FModuleHandle> Payload;
 
 	FModuleHandle Handle;
 
diff --git a/Engine/Plugins/Experimental/AnimNext/Source/AnimNext/Private/Component/AnimNextWorldSubsystem.cpp b/Engine/Plugins/Experimental/AnimNext/Source/AnimNext/Private/Component/AnimNextWorldSubsystem.cpp
--- a/Engine/Plugins/Experimental/AnimNext/Source/AnimNext/Private/Component/AnimNextWorldSubsystem.cpp
+++ b/Engine/Plugins/Experimental/AnimNext/Source/AnimNext/Private/Component/AnimNextWorldSubsystem.cpp
@@ -139,6 +139,44 @@ void UAnimNextWorldSubsystem::FlushPendingActions()
 				}
 #endif
 				break;
+
+			case FModulePendingAction::EType::AddPrerequisite:
+				{					
+					FAnimNextModuleInstance* SubsequentInstance = Instances.TryGet(PendingAction.Handle);
+					FAnimNextModuleInstance* PrerequisiteInstance = Instances.TryGet(PendingAction.Payload.Get<UE::AnimNext::FModuleHandle>());
+					if(PrerequisiteInstance != nullptr && SubsequentInstance != nullptr)
+					{
+						SubsequentInstance->AddPrerequisite(*PrerequisiteInstance);
+					}
+					else
+					{
+						UE_LOGFMT(LogAnimation, Warning, "FlushPendingActions: Invalid Module handle(s) provided for AddPrerequisite action SubsequentHandle: %i PrerequisiteInstance: %i", PendingAction.Handle.IsValid(), PendingAction.Payload.Get<UE::AnimNext::FModuleHandle>().IsValid());
+					}
+				
+					break;
+				}
+			case FModulePendingAction::EType::RemovePrerequisite:
+				{					
+					FAnimNextModuleInstance* SubsequentInstance = Instances.TryGet(PendingAction.Handle);
+					FAnimNextModuleInstance* PrerequisiteInstance = Instances.TryGet(PendingAction.Payload.Get<UE::AnimNext::FModuleHandle>());
+					if(PrerequisiteInstance != nullptr && SubsequentInstance != nullptr)
+					{
+						if (!SubsequentInstance->IsPrerequisite(*PrerequisiteInstance))
+						{							
+							SubsequentInstance->RemovePrerequisite(*PrerequisiteInstance);
+						}
+						else
+						{
+							UE_LOGFMT(LogAnimation, Warning, "FlushPendingActions: Trying to remove prerequisite ModuleInstance which isn't actually a prerequisite");
+						}
+					}
+					else
+					{
+						UE_LOGFMT(LogAnimation, Warning, "FlushPendingActions: Invalid Module handle(s) provided for RemovePrerequisite action SubsequentHandle: %i PrerequisiteInstance: %i", PendingAction.Handle.IsValid(), PendingAction.Payload.Get<UE::AnimNext::FModuleHandle>().IsValid());
+					}
+				
+					break;
+				}		
 			default:
 				checkNoEntry();
 				break;
@@ -326,14 +364,7 @@ void UAnimNextWorldSubsystem::AddModuleDependencyHandle(UE::AnimNext::FModuleHan
 		return;
 	}
 
-	FAnimNextModuleInstance* PrerequisiteInstance = Instances.TryGet(InPrerequisiteHandle);
-	FAnimNextModuleInstance* SubsequentInstance = Instances.TryGet(InSubsequentHandle);
-	if(PrerequisiteInstance == nullptr || SubsequentInstance == nullptr)
-	{
-		return;
-	}
-
-	SubsequentInstance->AddPrerequisite(*PrerequisiteInstance);
+	PendingActions.Emplace(InSubsequentHandle, FModulePendingAction::EType::AddPrerequisite, InPrerequisiteHandle);
 }
 
 void UAnimNextWorldSubsystem::RemoveModuleDependencyHandle(UE::AnimNext::FModuleHandle InPrerequisiteHandle, UE::AnimNext::FModuleHandle InSubsequentHandle)
@@ -345,14 +376,7 @@ void UAnimNextWorldSubsystem::RemoveModuleDependencyHandle(UE::AnimNext::FModule
 		return;
 	}
 
-	FAnimNextModuleInstance* PrerequisiteInstance = Instances.TryGet(InPrerequisiteHandle);
-	FAnimNextModuleInstance* SubsequentInstance = Instances.TryGet(InSubsequentHandle);
-	if(PrerequisiteInstance == nullptr || SubsequentInstance == nullptr)
-	{
-		return;
-	}
-
-	SubsequentInstance->RemovePrerequisite(*PrerequisiteInstance);
+	PendingActions.Emplace(InSubsequentHandle, FModulePendingAction::EType::RemovePrerequisite, InPrerequisiteHandle);
 }
 
 const FTickFunction* UAnimNextWorldSubsystem::FindTickFunctionHandle(UE::AnimNext::FModuleHandle InHandle, FName InEventName) const
