Allow setting dependencies between anim next components in blueprint.

diff --git a/Engine/Plugins/Experimental/AnimNext/Source/AnimNext/Internal/Component/AnimNextWorldSubsystem.h b/Engine/Plugins/Experimental/AnimNext/Source/AnimNext/Internal/Component/AnimNextWorldSubsystem.h
--- a/Engine/Plugins/Experimental/AnimNext/Source/AnimNext/Internal/Component/AnimNextWorldSubsystem.h
+++ b/Engine/Plugins/Experimental/AnimNext/Source/AnimNext/Internal/Component/AnimNextWorldSubsystem.h
@@ -171,6 +171,22 @@ protected:
 	// @param	InDependency		The kind of dependency to remove
 	void RemoveDependencyHandle(UE::AnimNext::FModuleHandle InHandle, UObject* InObject, FTickFunction& InTickFunction, FName InEventName, EDependency InDependency);
 
+	// Add a dependency on a module event to the specified module event
+	// @param	InHandle			The handle of the module whose dependencies are being modified
+	// @param	InEventName			The event to add the dependency to/from
+	// @param	OtherHandle			The handle of the other module whose dependencies are being modified
+	// @param	OtherEventName		The other module's event to add the dependency to/from
+	// @param	InDependency		The kind of dependency to add
+	void AddModuleEventDependencyHandle(UE::AnimNext::FModuleHandle InHandle, FName InEventName, UE::AnimNext::FModuleHandle OtherHandle, FName OtherEventName, EDependency InDependency);
+
+	// Remove a dependency on a module event from the specified module event
+	// @param	InHandle			The handle of the module whose dependencies are being modified
+	// @param	InEventName			The event to remove the dependency to/from
+	// @param	OtherHandle			The handle of the other module whose dependencies are being modified
+	// @param	OtherEventName		The other module's event to remove the dependency to/from
+	// @param	InDependency		The kind of dependency to remove
+	void RemoveModuleEventDependencyHandle(UE::AnimNext::FModuleHandle InHandle, FName InEventName, UE::AnimNext::FModuleHandle OtherHandle, FName OtherEventName, EDependency InDependency);
+
 protected:
 	friend struct FRigUnit_AnimNextInitializeEvent;
 	friend struct UE::AnimNext::FModuleEndTickFunction;
diff --git a/Engine/Plugins/Experimental/AnimNext/Source/AnimNext/Private/Component/AnimNextComponent.cpp b/Engine/Plugins/Experimental/AnimNext/Source/AnimNext/Private/Component/AnimNextComponent.cpp
--- a/Engine/Plugins/Experimental/AnimNext/Source/AnimNext/Private/Component/AnimNextComponent.cpp
+++ b/Engine/Plugins/Experimental/AnimNext/Source/AnimNext/Private/Component/AnimNextComponent.cpp
@@ -950,6 +950,70 @@ void UAnimNextComponent::RemoveSubsequent(UObject* InObject, FTickFunction& InTi
 	}
 }
 
+void UAnimNextComponent::AddModuleEventPrerequisite(FName InEventName, UAnimNextComponent* OtherAnimNextComponent, FName OtherEventName)
+{
+	if (!OtherAnimNextComponent)
+	{
+		UE_LOG(LogAnimation, Warning, TEXT("UAnimNextComponent::AddModuleEventPrerequisite called with null OtherAnimNextComponent"));
+	}
+	else if (OtherAnimNextComponent == this)
+	{
+		UE_LOG(LogAnimation, Warning, TEXT("UAnimNextComponent::AddModuleEventPrerequisite called using the same component"));
+	}
+	else if (Subsystem)
+	{
+		Subsystem->AddModuleEventDependency(this, InEventName, OtherAnimNextComponent, OtherEventName, UAnimNextWorldSubsystem::EDependency::Prerequisite);
+	}
+}
+
+void UAnimNextComponent::AddModuleEventSubsequent(FName InEventName, UAnimNextComponent* OtherAnimNextComponent, FName OtherEventName)
+{
+	if (!OtherAnimNextComponent)
+	{
+		UE_LOG(LogAnimation, Warning, TEXT("UAnimNextComponent::AddModuleEventSubsequent called with null OtherAnimNextComponent"));
+	}
+	else if (OtherAnimNextComponent == this)
+	{
+		UE_LOG(LogAnimation, Warning, TEXT("UAnimNextComponent::AddModuleEventPrerequisite called using the same component"));
+	}
+	else if (Subsystem)
+	{
+		Subsystem->AddModuleEventDependency(this, InEventName, OtherAnimNextComponent, OtherEventName, UAnimNextWorldSubsystem::EDependency::Subsequent);
+	}
+}
+
+void UAnimNextComponent::RemoveModuleEventPrerequisite(FName InEventName, UAnimNextComponent* OtherAnimNextComponent, FName OtherEventName)
+{
+	if (!OtherAnimNextComponent)
+	{
+		UE_LOG(LogAnimation, Warning, TEXT("UAnimNextComponent::RemoveModuleEventPrerequisite called with null OtherAnimNextComponent"));
+	}
+	else if (OtherAnimNextComponent == this)
+	{
+		UE_LOG(LogAnimation, Warning, TEXT("UAnimNextComponent::AddModuleEventPrerequisite called using the same component"));
+	}
+	else if (Subsystem)
+	{
+		Subsystem->RemoveModuleEventDependency(this, InEventName, OtherAnimNextComponent, OtherEventName, UAnimNextWorldSubsystem::EDependency::Prerequisite);
+	}
+}
+
+void UAnimNextComponent::RemoveModuleEventSubsequent(FName InEventName, UAnimNextComponent* OtherAnimNextComponent, FName OtherEventName)
+{
+	if (!OtherAnimNextComponent)
+	{
+		UE_LOG(LogAnimation, Warning, TEXT("UAnimNextComponent::RemoveModuleEventSubsequent called with null OtherAnimNextComponent"));
+	}
+	else if (OtherAnimNextComponent == this)
+	{
+		UE_LOG(LogAnimation, Warning, TEXT("UAnimNextComponent::AddModuleEventPrerequisite called using the same component"));
+	}
+	else if (Subsystem)
+	{
+		Subsystem->RemoveModuleEventDependency(this, InEventName, OtherAnimNextComponent, OtherEventName, UAnimNextWorldSubsystem::EDependency::Subsequent);
+	}
+}
+
 FAnimNextModuleHandle UAnimNextComponent::BlueprintGetModuleHandle() const
 {
 	return FAnimNextModuleHandle(ModuleHandle);
diff --git a/Engine/Plugins/Experimental/AnimNext/Source/AnimNext/Private/Component/AnimNextComponentWorldSubsystem.cpp b/Engine/Plugins/Experimental/AnimNext/Source/AnimNext/Private/Component/AnimNextComponentWorldSubsystem.cpp
--- a/Engine/Plugins/Experimental/AnimNext/Source/AnimNext/Private/Component/AnimNextComponentWorldSubsystem.cpp
+++ b/Engine/Plugins/Experimental/AnimNext/Source/AnimNext/Private/Component/AnimNextComponentWorldSubsystem.cpp
@@ -91,6 +91,22 @@ void UAnimNextComponentWorldSubsystem::RemoveDependency(UAnimNextComponent* InCo
 	RemoveDependencyHandle(InComponent->ModuleHandle, InObject, InTickFunction, InEventName, InDependency);
 }
 
+void UAnimNextComponentWorldSubsystem::AddModuleEventDependency(UAnimNextComponent* InComponent, FName InEventName, UAnimNextComponent* OtherComponent, FName OtherEventName, EDependency InDependency)
+{
+	using namespace UE::AnimNext;
+	check(IsInGameThread());
+	check(InComponent && OtherComponent && InComponent != OtherComponent);
+	AddModuleEventDependencyHandle(InComponent->ModuleHandle, InEventName, OtherComponent->ModuleHandle, OtherEventName, InDependency);
+}
+
+void UAnimNextComponentWorldSubsystem::RemoveModuleEventDependency(UAnimNextComponent* InComponent, FName InEventName, UAnimNextComponent* OtherComponent, FName OtherEventName, EDependency InDependency)
+{
+	using namespace UE::AnimNext;
+	check(IsInGameThread());
+	check(InComponent && OtherComponent && InComponent != OtherComponent);
+	RemoveModuleEventDependencyHandle(InComponent->ModuleHandle, InEventName, OtherComponent->ModuleHandle, OtherEventName, InDependency);
+}
+
 #if WITH_EDITOR
 
 void UAnimNextComponentWorldSubsystem::OnModuleCompiled(UAnimNextModule* InModule)
diff --git a/Engine/Plugins/Experimental/AnimNext/Source/AnimNext/Private/Component/AnimNextComponentWorldSubsystem.h b/Engine/Plugins/Experimental/AnimNext/Source/AnimNext/Private/Component/AnimNextComponentWorldSubsystem.h
--- a/Engine/Plugins/Experimental/AnimNext/Source/AnimNext/Private/Component/AnimNextComponentWorldSubsystem.h
+++ b/Engine/Plugins/Experimental/AnimNext/Source/AnimNext/Private/Component/AnimNextComponentWorldSubsystem.h
@@ -68,6 +68,22 @@ class UAnimNextComponentWorldSubsystem : public UAnimNextWorldSubsystem
 	// @param	InEventName			The event to remove the dependency from
 	void RemoveDependency(UAnimNextComponent* InComponent, UObject* InObject, FTickFunction& InTickFunction, FName InEventName, EDependency InDependency);
 
+	// Add a dependency on a module event to the specified module event
+	// @param	InComponent			The component whose dependencies are being modified
+	// @param	InEventName			The event to add the dependency to/from
+	// @param	OtherComponent		The other component whose dependencies are being modified
+	// @param	OtherEventName		The other module's event to add the dependency to/from
+	// @param	InDependency		The kind of dependency to add
+	void AddModuleEventDependency(UAnimNextComponent* InComponent, FName InEventName, UAnimNextComponent* OtherComponent, FName OtherEventName, EDependency InDependency);
+
+	// Remove a dependency on a module event to the specified module event
+	// @param	InComponent			The component whose dependencies are being modified
+	// @param	InEventName			The event to remove the dependency to/from
+	// @param	OtherComponent		The other component whose dependencies are being modified
+	// @param	OtherEventName		The other module's event to remove the dependency to/from
+	// @param	InDependency		The kind of dependency to remove
+	void RemoveModuleEventDependency(UAnimNextComponent* InComponent, FName InEventName, UAnimNextComponent* OtherComponent, FName OtherEventName, EDependency InDependency);
+
 #if WITH_EDITOR
 	// UAnimNextWorldSubsystem interface
 	virtual void OnModuleCompiled(UAnimNextModule* InModule) override;
diff --git a/Engine/Plugins/Experimental/AnimNext/Source/AnimNext/Private/Component/AnimNextWorldSubsystem.cpp b/Engine/Plugins/Experimental/AnimNext/Source/AnimNext/Private/Component/AnimNextWorldSubsystem.cpp
--- a/Engine/Plugins/Experimental/AnimNext/Source/AnimNext/Private/Component/AnimNextWorldSubsystem.cpp
+++ b/Engine/Plugins/Experimental/AnimNext/Source/AnimNext/Private/Component/AnimNextWorldSubsystem.cpp
@@ -451,4 +451,42 @@ void UAnimNextWorldSubsystem::RemoveDependencyHandle(UE::AnimNext::FModuleHandle
 			TickFunction->RemoveSubsequent(InObject, InTickFunction);
 		}
 	}
-}
\ No newline at end of file
+}
+
+void UAnimNextWorldSubsystem::AddModuleEventDependencyHandle(UE::AnimNext::FModuleHandle InHandle, FName InEventName, UE::AnimNext::FModuleHandle OtherHandle, FName OtherEventName, EDependency InDependency)
+{
+	using namespace UE::AnimNext;
+	check(IsInGameThread());
+	check(InHandle != OtherHandle);
+	if (IsValidHandle(OtherHandle))
+	{
+		FAnimNextModuleInstance& Instance = Instances.Get(OtherHandle);
+		FModuleEventTickFunction* TickFunction = Instance.FindTickFunction(OtherEventName);
+		if(TickFunction == nullptr)
+		{
+			UE_LOGFMT(LogAnimation, Warning, "AddDependencyHandle: Could not find event '{EventName}' in module '{ModuleName}'", OtherEventName, Instance.GetDataInterfaceName());
+			return;
+		}
+
+		AddDependencyHandle(InHandle, Instance.GetObject(), *TickFunction, InEventName, InDependency);
+	}
+}
+
+void UAnimNextWorldSubsystem::RemoveModuleEventDependencyHandle(UE::AnimNext::FModuleHandle InHandle, FName InEventName, UE::AnimNext::FModuleHandle OtherHandle, FName OtherEventName, EDependency InDependency)
+{
+	using namespace UE::AnimNext;
+	check(IsInGameThread());
+	check(InHandle != OtherHandle);
+	if (IsValidHandle(InHandle))
+	{
+		FAnimNextModuleInstance& Instance = Instances.Get(InHandle);
+		FModuleEventTickFunction* TickFunction = Instance.FindTickFunction(InEventName);
+		if(TickFunction == nullptr)
+		{
+			UE_LOGFMT(LogAnimation, Warning, "RemoveDependencyHandle: Could not find event '{EventName}' in module '{ModuleName}'", OtherEventName, Instance.GetDataInterfaceName());
+			return;
+		}
+
+		RemoveDependencyHandle(InHandle, Instance.GetObject(), *TickFunction, InEventName, InDependency);
+	}
+}
diff --git a/Engine/Plugins/Experimental/AnimNext/Source/AnimNext/Public/Component/AnimNextComponent.h b/Engine/Plugins/Experimental/AnimNext/Source/AnimNext/Public/Component/AnimNextComponent.h
--- a/Engine/Plugins/Experimental/AnimNext/Source/AnimNext/Public/Component/AnimNextComponent.h
+++ b/Engine/Plugins/Experimental/AnimNext/Source/AnimNext/Public/Component/AnimNextComponent.h
@@ -233,6 +233,34 @@ public:
 	// @param	InEventName			The event to remove the dependency from
 	ANIMNEXT_API void RemoveSubsequent(UObject* InObject, FTickFunction& InTickFunction, FName InEventName);
 
+	// Add a prerequisite anim next event dependency to the specified event
+	// @param	InEventName			The event name in this component
+	// @param	InTickFunction		The other component we want a prerequisite on
+	// @param	InEventName			The other component's event name
+	UFUNCTION(BlueprintCallable, Category = "AnimNext")
+	ANIMNEXT_API void AddModuleEventPrerequisite(FName InEventName, UAnimNextComponent* OtherAnimNextComponent, FName OtherEventName);
+
+	// Add a subsequent anim next event dependency to the specified event
+	// @param	InEventName			The event name in this component
+	// @param	InTickFunction		The other component we want to add a prerequisite to
+	// @param	InEventName			The other component's event name
+	UFUNCTION(BlueprintCallable, Category = "AnimNext")
+	ANIMNEXT_API void AddModuleEventSubsequent(FName InEventName, UAnimNextComponent* OtherAnimNextComponent, FName OtherEventName);
+
+	// Remove a prerequisite anim next event dependency from the specified event
+	// @param	InEventName			The event name in this component
+	// @param	InTickFunction		The other component we want to remove a prerequisite from
+	// @param	InEventName			The other component's event name
+	UFUNCTION(BlueprintCallable, Category = "AnimNext")
+	ANIMNEXT_API void RemoveModuleEventPrerequisite(FName InEventName, UAnimNextComponent* OtherAnimNextComponent, FName OtherEventName);
+
+	// Remove a subsequent anim next event dependency from the specified event
+	// @param	InEventName			The event name in this component
+	// @param	InTickFunction		The other component we want to remove a prerequisite to
+	// @param	InEventName			The other component's event name
+	UFUNCTION(BlueprintCallable, Category = "AnimNext")
+	ANIMNEXT_API void RemoveModuleEventSubsequent(FName InEventName, UAnimNextComponent* OtherAnimNextComponent, FName OtherEventName);
+
 	UFUNCTION(BlueprintCallable, Category = "AnimNext", meta=(DisplayName = "Get Module Handle"))
 	ANIMNEXT_API FAnimNextModuleHandle BlueprintGetModuleHandle() const;
 
