AnimNext: Fix animated bounds not being sent to the renderer

Bounds were updated if the component moved, but not if animation changed them, resulting in shadow artifacts

diff --git a/Engine/Plugins/Experimental/AnimNext/Source/AnimNext/Private/Component/SkinnedMeshComponentExtensions.cpp b/Engine/Plugins/Experimental/AnimNext/Source/AnimNext/Private/Component/SkinnedMeshComponentExtensions.cpp
--- a/Engine/Plugins/Experimental/AnimNext/Source/AnimNext/Private/Component/SkinnedMeshComponentExtensions.cpp
+++ b/Engine/Plugins/Experimental/AnimNext/Source/AnimNext/Private/Component/SkinnedMeshComponentExtensions.cpp
@@ -30,8 +30,9 @@ void FSkinnedMeshComponentExtensions::CompleteAndDispatch(
 		InComponent->InvalidateCachedBounds();
 		InComponent->UpdateBounds();
 
-		// Send updated transforms to the renderer
+		// Send updated transforms & bounds to the renderer
 		InComponent->SendRenderDynamicData_Concurrent();
+		InComponent->SendRenderTransform_Concurrent();
 
 		if (InComponent->bHasSocketAttachments)
 		{
