Improve warning log message to include offending asset

diff --git a/Engine/Plugins/Experimental/AnimNext/Source/AnimNext/Private/Graph/RigUnit_AnimNextWriteSkeletalMeshComponentPose.cpp b/Engine/Plugins/Experimental/AnimNext/Source/AnimNext/Private/Graph/RigUnit_AnimNextWriteSkeletalMeshComponentPose.cpp
--- a/Engine/Plugins/Experimental/AnimNext/Source/AnimNext/Private/Graph/RigUnit_AnimNextWriteSkeletalMeshComponentPose.cpp
+++ b/Engine/Plugins/Experimental/AnimNext/Source/AnimNext/Private/Graph/RigUnit_AnimNextWriteSkeletalMeshComponentPose.cpp
@@ -7,6 +7,7 @@
 #include "Component/SkinnedMeshComponentExtensions.h"
 #include "Components/SkeletalMeshComponent.h"
 #include "Engine/SkeletalMesh.h"
+#include "Logging/StructuredLog.h"
 #include "Module/AnimNextSkeletalMeshComponentReferenceComponent.h"
 
 #include UE_INLINE_GENERATED_CPP_BY_NAME(RigUnit_AnimNextWriteSkeletalMeshComponentPose)
@@ -50,7 +51,8 @@ FRigUnit_AnimNextWriteSkeletalMeshComponentPose_Execute()
 	// We cannot write to the skeletal mesh component if it is using anim BP
 	if(OutputComponent->bEnableAnimation)
 	{
-		UE_LOG(LogAnimation, Warning, TEXT("AnimNext: Could not write to skeletal mesh component - bEnableAnimation is true"));
+		UE_LOGFMT(LogAnimation, Warning, "AnimNext: Could not write to skeletal mesh component - bEnableAnimation is true [SK: {SkeletalMesh}][Owner: {Owner}]",
+			*SkeletalMesh->GetName(), OutputComponent->GetOuter() != nullptr ? *OutputComponent->GetOuter()->GetName() : TEXT("*NO OWNER*"));
 		return;
 	}
 
@@ -61,7 +63,8 @@ FRigUnit_AnimNextWriteSkeletalMeshComponentPose_Execute()
 	// You should remap the pose explicitly if this is what you need
 	if (&RefPose != &Pose.LODPose.GetRefPose())
 	{
-		UE_LOG(LogAnimation, Warning, TEXT("AnimNext: Could not write to skeletal mesh component - The input pose and the skeletal mesh component use different reference poses"));
+		UE_LOGFMT(LogAnimation, Warning, "AnimNext: Could not write to skeletal mesh component - The input pose and the skeletal mesh component use different reference poses [SK: {SkeletalMesh}][Owner: {Owner}]",
+			*SkeletalMesh->GetName(), OutputComponent->GetOuter() != nullptr ? *OutputComponent->GetOuter()->GetName() : TEXT("*NO OWNER*"));
 		return;
 	}
 
