Fix crash close animnext editor due to dangling delegate.

diff --git a/Engine/Plugins/Experimental/AnimNext/Source/AnimNextEditor/Private/Workspace/AnimNextWorkspaceEditorMode.cpp b/Engine/Plugins/Experimental/AnimNext/Source/AnimNextEditor/Private/Workspace/AnimNextWorkspaceEditorMode.cpp
--- a/Engine/Plugins/Experimental/AnimNext/Source/AnimNextEditor/Private/Workspace/AnimNextWorkspaceEditorMode.cpp
+++ b/Engine/Plugins/Experimental/AnimNext/Source/AnimNextEditor/Private/Workspace/AnimNextWorkspaceEditorMode.cpp
@@ -1,4 +1,4 @@
-﻿// Copyright Epic Games, Inc. All Rights Reserved.
+// Copyright Epic Games, Inc. All Rights Reserved.
 
 #include "AnimNextWorkspaceEditorMode.h"
 
@@ -64,6 +64,14 @@ void UAnimNextWorkspaceEditorMode::Exit()
 {
 	Super::Exit();
 
+	for (TPair<FObjectKey, TSharedRef<UE::AnimNext::Editor::IAssetCompilationHandler>>& AssetCompiler : AssetCompilers)
+	{
+		TSharedRef<UE::AnimNext::Editor::IAssetCompilationHandler> AssetCompilerRef = AssetCompiler.Value;
+		AssetCompilerRef->OnCompileStatusChanged().Unbind();
+	}
+
+	AssetCompilers.Reset();
+
 	for(TObjectKey<UObject> Asset : WeakAssets)
 	{
 		UAnimNextRigVMAsset* AnimNextRigVMAsset = Cast<UAnimNextRigVMAsset>(Asset.ResolveObjectPtr());
