Fix use after free AnimNext when updating outliner due to array realloc

#tests Standalone

diff --git a/Engine/Plugins/Experimental/AnimNext/Source/AnimNextUncookedOnly/Internal/UncookedOnlyUtils.h b/Engine/Plugins/Experimental/AnimNext/Source/AnimNextUncookedOnly/Internal/UncookedOnlyUtils.h
--- a/Engine/Plugins/Experimental/AnimNext/Source/AnimNextUncookedOnly/Internal/UncookedOnlyUtils.h
+++ b/Engine/Plugins/Experimental/AnimNext/Source/AnimNextUncookedOnly/Internal/UncookedOnlyUtils.h
@@ -223,10 +223,11 @@ struct ANIMNEXTUNCOOKEDONLY_API FUtils
 	static void GetAssetPrivateFunctions(const UAnimNextRigVMAssetEditorData* InEditorData, FRigVMGraphFunctionHeaderArray& OutExports);
 
 	// Gets the asset-registry information needed for representing the contained data into the Workspace Outliner
+	// Note: We pass parent as index to avoid the ref from being invalidated due to realloc as the export array grows recursively
 	static void GetAssetOutlinerItems(const UAnimNextRigVMAssetEditorData* EditorData, FWorkspaceOutlinerItemExports& OutExports, FAssetRegistryTagsContext Context);
-	static void CreateSubGraphsOutlinerItemsRecursive(const UAnimNextRigVMAssetEditorData* EditorData, FWorkspaceOutlinerItemExports& OutExports, FWorkspaceOutlinerItemExport& ParentExport, URigVMEdGraph* RigVMEdGraph, FAssetRegistryTagsContext Context);
-	static void CreateFunctionLibraryOutlinerItemsRecursive(const UAnimNextRigVMAssetEditorData* EditorData, FWorkspaceOutlinerItemExports& OutExports, FWorkspaceOutlinerItemExport& ParentExport, const TArray<FRigVMGraphFunctionData>& PublicFunctions, const TArray<FRigVMGraphFunctionData>& PrivateFunctions);
-	static void CreateFunctionsOutlinerItemsRecursive(const UAnimNextRigVMAssetEditorData* EditorData, FWorkspaceOutlinerItemExports& OutExports, FWorkspaceOutlinerItemExport& ParentExport, const TArray<FRigVMGraphFunctionData>& Functions, bool bPublicFunctions);
+	static void CreateSubGraphsOutlinerItemsRecursive(const UAnimNextRigVMAssetEditorData* EditorData, FWorkspaceOutlinerItemExports& OutExports, FWorkspaceOutlinerItemExport& RootExport, int32 ParentExportIndex, URigVMEdGraph* RigVMEdGraph, FAssetRegistryTagsContext Context);
+	static void CreateFunctionLibraryOutlinerItemsRecursive(const UAnimNextRigVMAssetEditorData* EditorData, FWorkspaceOutlinerItemExports& OutExports, FWorkspaceOutlinerItemExport& RootExport, int32 ParentExportIndex, const TArray<FRigVMGraphFunctionData>& PublicFunctions, const TArray<FRigVMGraphFunctionData>& PrivateFunctions);
+	static void CreateFunctionsOutlinerItemsRecursive(const UAnimNextRigVMAssetEditorData* EditorData, FWorkspaceOutlinerItemExports& OutExports, FWorkspaceOutlinerItemExport& RootExport, int32 ParentExportIndex, const TArray<FRigVMGraphFunctionData>& Functions, bool bPublicFunctions);
 
 	// Attempts to determine the type from a parameter name
 	// If the name cannot be found, the returned type will be invalid
diff --git a/Engine/Plugins/Experimental/AnimNext/Source/AnimNextUncookedOnly/Private/UncookedOnlyUtils.cpp b/Engine/Plugins/Experimental/AnimNext/Source/AnimNextUncookedOnly/Private/UncookedOnlyUtils.cpp
--- a/Engine/Plugins/Experimental/AnimNext/Source/AnimNextUncookedOnly/Private/UncookedOnlyUtils.cpp
+++ b/Engine/Plugins/Experimental/AnimNext/Source/AnimNextUncookedOnly/Private/UncookedOnlyUtils.cpp
@@ -955,7 +955,7 @@ void FUtils::GetAssetOutlinerItems(const UAnimNextRigVMAssetEditorData* EditorDa
 			{
 				if (URigVMEdGraph* RigVMEdGraph = GraphInterface->GetEdGraph())
 				{
-					CreateSubGraphsOutlinerItemsRecursive(EditorData, OutExports, AssetIdentifier, RigVMEdGraph, Context);
+					CreateSubGraphsOutlinerItemsRecursive(EditorData, OutExports, AssetIdentifier, INDEX_NONE, RigVMEdGraph, Context);
 				}
 			}
 			else
@@ -968,16 +968,17 @@ void FUtils::GetAssetOutlinerItems(const UAnimNextRigVMAssetEditorData* EditorDa
 
 				if (URigVMEdGraph* RigVMEdGraph = GraphInterface->GetEdGraph())
 				{
-					CreateSubGraphsOutlinerItemsRecursive(EditorData, OutExports, Export, RigVMEdGraph, Context);
+					int32 ExportIndex = OutExports.Exports.Num() - 1;
+					CreateSubGraphsOutlinerItemsRecursive(EditorData, OutExports, AssetIdentifier, ExportIndex, RigVMEdGraph, Context);
 				}
 			}
 		}
 	}
 
-	CreateFunctionLibraryOutlinerItemsRecursive(EditorData, OutExports, AssetIdentifier, EditorData->GetRigVMGraphFunctionStore()->PublicFunctions, EditorData->GetRigVMGraphFunctionStore()->PrivateFunctions);
+	CreateFunctionLibraryOutlinerItemsRecursive(EditorData, OutExports, AssetIdentifier, INDEX_NONE, EditorData->GetRigVMGraphFunctionStore()->PublicFunctions, EditorData->GetRigVMGraphFunctionStore()->PrivateFunctions);
 }
 
-void ProcessPinAssetReferences(const URigVMPin* InPin, FWorkspaceOutlinerItemExports& OutExports, const FWorkspaceOutlinerItemExport& ParentExport)
+void ProcessPinAssetReferences(const URigVMPin* InPin, FWorkspaceOutlinerItemExports& OutExports, FWorkspaceOutlinerItemExport& RootExport, int32 ParentExportIndex)
 {
 	if (InPin)
 	{
@@ -987,6 +988,7 @@ void ProcessPinAssetReferences(const URigVMPin* InPin, FWorkspaceOutlinerItemExp
 			const FSoftObjectPath ObjectPath(InPin->GetDefaultValue());
 			if (ObjectPath.IsValid())
 			{
+				FWorkspaceOutlinerItemExport& ParentExport = ParentExportIndex == INDEX_NONE ? RootExport : OutExports.Exports[ParentExportIndex];
 				FWorkspaceOutlinerItemExport& ReferenceExport = OutExports.Exports.Add_GetRef(FWorkspaceOutlinerItemExport(FName(ObjectPath.ToString()), ParentExport));
 				ReferenceExport.GetData().InitializeAs<FWorkspaceOutlinerAssetReferenceItemData>();
 				ReferenceExport.GetData().GetMutable<FWorkspaceOutlinerAssetReferenceItemData>().ReferredObjectPath = ObjectPath;
@@ -995,13 +997,13 @@ void ProcessPinAssetReferences(const URigVMPin* InPin, FWorkspaceOutlinerItemExp
 
 		for (const URigVMPin* SubPin : InPin->GetSubPins())
 		{
-			ProcessPinAssetReferences(SubPin, OutExports, ParentExport);
+			ProcessPinAssetReferences(SubPin, OutExports, RootExport, ParentExportIndex);
 		}
 	}
 	
 }
 
-void FUtils::CreateSubGraphsOutlinerItemsRecursive(const UAnimNextRigVMAssetEditorData* EditorData, FWorkspaceOutlinerItemExports& OutExports, FWorkspaceOutlinerItemExport& ParentExport, URigVMEdGraph* RigVMEdGraph, FAssetRegistryTagsContext Context)
+void FUtils::CreateSubGraphsOutlinerItemsRecursive(const UAnimNextRigVMAssetEditorData* EditorData, FWorkspaceOutlinerItemExports& OutExports, FWorkspaceOutlinerItemExport& RootExport, int32 ParentExportIndex, URigVMEdGraph* RigVMEdGraph, FAssetRegistryTagsContext Context)
 {
 	if (RigVMEdGraph == nullptr)
 	{
@@ -1023,11 +1025,11 @@ void FUtils::CreateSubGraphsOutlinerItemsRecursive(const UAnimNextRigVMAssetEdit
 						{
 							if (ModelPin->GetDirection() == ERigVMPinDirection::Input)
 							{
-								auto HandlePin = [&OutExports, ParentExport](const URigVMPin* InPin)
+								auto HandlePin = [&OutExports, &RootExport, ParentExportIndex](const URigVMPin* InPin)
 								{
 									if (InPin->GetMetaData(TEXT("ExportAsReference")) == TEXT("true"))
 									{
-										ProcessPinAssetReferences(InPin, OutExports, ParentExport);
+										ProcessPinAssetReferences(InPin, OutExports, RootExport, ParentExportIndex);
 									}										
 								};
 								
@@ -1055,13 +1057,15 @@ void FUtils::CreateSubGraphsOutlinerItemsRecursive(const UAnimNextRigVMAssetEdit
 			{
 				URigVMCollapseNode* CollapseNode = CastChecked<URigVMCollapseNode>(EditorObject->GetModel()->GetOuter());
 
+				FWorkspaceOutlinerItemExport& ParentExport = ParentExportIndex == INDEX_NONE ? RootExport : OutExports.Exports[ParentExportIndex];
 				FWorkspaceOutlinerItemExport& Export = OutExports.Exports.Add_GetRef(FWorkspaceOutlinerItemExport(CollapseNode->GetFName(), ParentExport));
 				Export.GetData().InitializeAsScriptStruct(FAnimNextCollapseGraphOutlinerData::StaticStruct());
 			
 				FAnimNextCollapseGraphOutlinerData& FnGraphData = Export.GetData().GetMutable<FAnimNextCollapseGraphOutlinerData>();
 				FnGraphData.SoftEditorObject = EditorObject;
 
-				CreateSubGraphsOutlinerItemsRecursive(EditorData, OutExports, Export, EditorObject, Context);
+				int32 ExportIndex = OutExports.Exports.Num() - 1;
+				CreateSubGraphsOutlinerItemsRecursive(EditorData, OutExports, RootExport, ExportIndex, EditorObject, Context);
 			}
 		}
 	}
@@ -1081,6 +1085,7 @@ void FUtils::CreateSubGraphsOutlinerItemsRecursive(const UAnimNextRigVMAssetEdit
 				const URigVMLibraryNode* FunctionNode = EditorData->RigVMClient.GetFunctionLibrary()->FindFunction(FunctionIdentifier.GetFunctionFName());
 				if (URigVMGraph* ContainedGraph = FunctionNode->GetContainedGraph())
 				{
+					FWorkspaceOutlinerItemExport& ParentExport = ParentExportIndex == INDEX_NONE ? RootExport : OutExports.Exports[ParentExportIndex];
 					FWorkspaceOutlinerItemExport& Export = OutExports.Exports.Add_GetRef(FWorkspaceOutlinerItemExport(FunctionNode->GetFName(), ParentExport));
 					Export.GetData().InitializeAsScriptStruct(FAnimNextGraphFunctionOutlinerData::StaticStruct());
 					FAnimNextGraphFunctionOutlinerData& FnGraphData = Export.GetData().GetMutable<FAnimNextGraphFunctionOutlinerData>();
@@ -1090,7 +1095,8 @@ void FUtils::CreateSubGraphsOutlinerItemsRecursive(const UAnimNextRigVMAssetEdit
 						FnGraphData.SoftEditorObject = ContainedEdGraph;
 						FnGraphData.SoftEdGraphNode = EdNode;
 						
-						CreateSubGraphsOutlinerItemsRecursive(EditorData, OutExports, Export, ContainedEdGraph, Context);
+						int32 ExportIndex = OutExports.Exports.Num() - 1;
+						CreateSubGraphsOutlinerItemsRecursive(EditorData, OutExports, RootExport, ExportIndex, ContainedEdGraph, Context);
 					}
 				}
 			}
@@ -1098,16 +1104,16 @@ void FUtils::CreateSubGraphsOutlinerItemsRecursive(const UAnimNextRigVMAssetEdit
 	}
 }
 
-void FUtils::CreateFunctionLibraryOutlinerItemsRecursive(const UAnimNextRigVMAssetEditorData* EditorData, FWorkspaceOutlinerItemExports& OutExports, FWorkspaceOutlinerItemExport& ParentExport, const TArray<FRigVMGraphFunctionData>& PublicFunctions, const TArray<FRigVMGraphFunctionData>& PrivateFunctions)
+void FUtils::CreateFunctionLibraryOutlinerItemsRecursive(const UAnimNextRigVMAssetEditorData* EditorData, FWorkspaceOutlinerItemExports& OutExports, FWorkspaceOutlinerItemExport& RootExport, int32 ParentExportIndex, const TArray<FRigVMGraphFunctionData>& PublicFunctions, const TArray<FRigVMGraphFunctionData>& PrivateFunctions)
 {
 	if (PrivateFunctions.Num() > 0 || PublicFunctions.Num() > 0)
 	{
-		CreateFunctionsOutlinerItemsRecursive(EditorData, OutExports, ParentExport, PrivateFunctions, false);
-		CreateFunctionsOutlinerItemsRecursive(EditorData, OutExports, ParentExport, PublicFunctions, true);
+		CreateFunctionsOutlinerItemsRecursive(EditorData, OutExports, RootExport, ParentExportIndex, PrivateFunctions, false);
+		CreateFunctionsOutlinerItemsRecursive(EditorData, OutExports, RootExport, ParentExportIndex, PublicFunctions, true);
 	}
 }
 
-void FUtils::CreateFunctionsOutlinerItemsRecursive(const UAnimNextRigVMAssetEditorData* EditorData, FWorkspaceOutlinerItemExports& OutExports, FWorkspaceOutlinerItemExport& ParentExport, const TArray<FRigVMGraphFunctionData>& Functions, bool bPublicFunctions)
+void FUtils::CreateFunctionsOutlinerItemsRecursive(const UAnimNextRigVMAssetEditorData* EditorData, FWorkspaceOutlinerItemExports& OutExports, FWorkspaceOutlinerItemExport& RootExport, int32 ParentExportIndex, const TArray<FRigVMGraphFunctionData>& Functions, bool bPublicFunctions)
 {
 	for (const FRigVMGraphFunctionData& FunctionData : Functions)
 	{
@@ -1118,6 +1124,7 @@ void FUtils::CreateFunctionsOutlinerItemsRecursive(const UAnimNextRigVMAssetEdit
 			{
 				if (URigVMEdGraph* EditorObject = Cast<URigVMEdGraph>(EditorData->GetEditorObjectForRigVMGraph(ContainedModelGraph)))
 				{
+					FWorkspaceOutlinerItemExport& ParentExport = ParentExportIndex == INDEX_NONE ? RootExport : OutExports.Exports[ParentExportIndex];
 					FWorkspaceOutlinerItemExport& Export = OutExports.Exports.Add_GetRef(FWorkspaceOutlinerItemExport(FunctionData.Header.Name, ParentExport));
 
 					Export.GetData().InitializeAsScriptStruct(FAnimNextGraphFunctionOutlinerData::StaticStruct());
