AnimNext: Fix enum variable types being unselectable

Also fix the color of enums in pin type selectors

diff --git a/Engine/Plugins/Experimental/AnimNext/Source/AnimNextUncookedOnly/Private/UncookedOnlyUtils.cpp b/Engine/Plugins/Experimental/AnimNext/Source/AnimNextUncookedOnly/Private/UncookedOnlyUtils.cpp
--- a/Engine/Plugins/Experimental/AnimNext/Source/AnimNextUncookedOnly/Private/UncookedOnlyUtils.cpp
+++ b/Engine/Plugins/Experimental/AnimNext/Source/AnimNextUncookedOnly/Private/UncookedOnlyUtils.cpp
@@ -532,7 +532,15 @@ FAnimNextParamType FUtils::GetParamTypeFromPinType(const FEdGraphPinType& InPinT
 	}
 	else if (InPinType.PinCategory == UEdGraphSchema_K2::PC_Byte)
 	{
-		ValueType = FAnimNextParamType::EValueType::Byte;
+		if (UEnum* Enum = Cast<UEnum>(InPinType.PinSubCategoryObject.Get()))
+		{
+			ValueType = FAnimNextParamType::EValueType::Enum;
+			ValueTypeObject = Enum;
+		}
+		else
+		{
+			ValueType = FAnimNextParamType::EValueType::Byte;
+		}
 	}
 	else if (InPinType.PinCategory == UEdGraphSchema_K2::PC_Int)
 	{
@@ -580,7 +588,7 @@ FAnimNextParamType FUtils::GetParamTypeFromPinType(const FEdGraphPinType& InPinT
 	else if (InPinType.PinCategory == UEdGraphSchema_K2::PC_Enum)
 	{
 		ValueType = FAnimNextParamType::EValueType::Enum;
-		ValueTypeObject = InPinType.PinSubCategoryObject.Get();
+		ValueTypeObject = Cast<UEnum>(InPinType.PinSubCategoryObject.Get());
 		ensure(ValueTypeObject);
 	}
 	else if (InPinType.PinCategory == UEdGraphSchema_K2::PC_Struct)
diff --git a/Engine/Source/Editor/BlueprintGraph/Private/EdGraphSchema_K2.cpp b/Engine/Source/Editor/BlueprintGraph/Private/EdGraphSchema_K2.cpp
--- a/Engine/Source/Editor/BlueprintGraph/Private/EdGraphSchema_K2.cpp
+++ b/Engine/Source/Editor/BlueprintGraph/Private/EdGraphSchema_K2.cpp
@@ -3059,7 +3059,7 @@ FLinearColor UEdGraphSchema_K2::GetPinTypeColor(const FEdGraphPinType& PinType)
 	{
 		return Settings->BooleanPinTypeColor;
 	}
-	else if (TypeName == PC_Byte)
+	else if (TypeName == PC_Byte || TypeName == PC_Enum)
 	{
 		return Settings->BytePinTypeColor;
 	}
