Improved ensure to cover both external/internal entry setups in AnimNext assets

diff --git a/Engine/Plugins/Experimental/AnimNext/Source/AnimNextUncookedOnly/Private/AnimNextRigVMAssetEditorData.cpp b/Engine/Plugins/Experimental/AnimNext/Source/AnimNextUncookedOnly/Private/AnimNextRigVMAssetEditorData.cpp
--- a/Engine/Plugins/Experimental/AnimNext/Source/AnimNextUncookedOnly/Private/AnimNextRigVMAssetEditorData.cpp
+++ b/Engine/Plugins/Experimental/AnimNext/Source/AnimNextUncookedOnly/Private/AnimNextRigVMAssetEditorData.cpp
@@ -252,7 +252,7 @@ void UAnimNextRigVMAssetEditorData::PostLoadExternalPackages()
 	}
 
 	// Internal entries should be empty if we are externally packaged
-	ensure(bUsesExternalPackages == InternalEntries.IsEmpty());
+	ensure(!bUsesExternalPackages || InternalEntries.IsEmpty());
 
 	// Copy any internal entries to the main entries array
 	Entries.Append(InternalEntries);
