AnimNext python script output from editor actions/operations

diff --git a/Engine/Plugins/Experimental/AnimNext/Source/AnimNextEditor/Private/AnimNextEditorModule.cpp b/Engine/Plugins/Experimental/AnimNext/Source/AnimNextEditor/Private/AnimNextEditorModule.cpp
--- a/Engine/Plugins/Experimental/AnimNext/Source/AnimNextEditor/Private/AnimNextEditorModule.cpp
+++ b/Engine/Plugins/Experimental/AnimNext/Source/AnimNextEditor/Private/AnimNextEditorModule.cpp
@@ -758,7 +758,7 @@ void FAnimNextEditorModule::RegisterWorkspaceDocumentTypes(Workspace::IWorkspace
 		{
 			Controller->RelinkSourceAndTargetPins(NodesToRemove[0], true);;
 		}
-		Controller->RemoveNodes(NodesToRemove, true);
+		Controller->RemoveNodes(NodesToRemove, true, true);
 		Controller->CloseUndoBracket();
 	});
 	GraphArgs.OnCanCopySelectedNodes = Workspace::FOnCanPerformActionOnSelectedNodes::CreateLambda([this](const Workspace::FWorkspaceEditorContext& InContext, const FGraphPanelSelectionSet& InSelectedNodes)
@@ -819,7 +819,7 @@ void FAnimNextEditorModule::RegisterWorkspaceDocumentTypes(Workspace::IWorkspace
 
 			Controller->OpenUndoBracket(TEXT("Pasted Nodes."));
 
-			if (UE::RigVM::Editor::Tools::PasteNodes(InPasteLocation, TextToImport, Controller, RigVMEdGraph->GetModel(), RigVMClientHost->GetLocalFunctionLibrary(), RigVMClientHost->GetRigVMGraphFunctionHost()))
+			if (UE::RigVM::Editor::Tools::PasteNodes(InPasteLocation, TextToImport, Controller, RigVMEdGraph->GetModel(), RigVMClientHost->GetLocalFunctionLibrary(), RigVMClientHost->GetRigVMGraphFunctionHost(), true, true))
 			{
 				Controller->CloseUndoBracket();
 			}
diff --git a/Engine/Plugins/Experimental/AnimNext/Source/AnimNextEditor/Private/Common/AnimNextCollapseNodeItemDetails.cpp b/Engine/Plugins/Experimental/AnimNext/Source/AnimNextEditor/Private/Common/AnimNextCollapseNodeItemDetails.cpp
--- a/Engine/Plugins/Experimental/AnimNext/Source/AnimNextEditor/Private/Common/AnimNextCollapseNodeItemDetails.cpp
+++ b/Engine/Plugins/Experimental/AnimNext/Source/AnimNextEditor/Private/Common/AnimNextCollapseNodeItemDetails.cpp
@@ -87,7 +87,7 @@ void FAnimNextCollapseNodeItemDetails::Delete(TConstArrayView<FWorkspaceOutliner
 								if (URigVMEdGraph* ContainerEdGraph = Cast<URigVMEdGraph>(ClientHost->GetEditorObjectForRigVMGraph(ContainerGraph)))
 								{
 									FScopedTransaction Transaction(LOCTEXT("DeleteCollapseNodeInOutliner", "Delete Collapse Node"));
-									ContainerEdGraph->GetController()->RemoveNode(CollapseNode);
+									ContainerEdGraph->GetController()->RemoveNode(CollapseNode, true, true);
 								}
 							}
 						}
diff --git a/Engine/Plugins/Experimental/AnimNext/Source/AnimNextEditor/Private/Common/AnimNextFunctionItemDetails.cpp b/Engine/Plugins/Experimental/AnimNext/Source/AnimNextEditor/Private/Common/AnimNextFunctionItemDetails.cpp
--- a/Engine/Plugins/Experimental/AnimNext/Source/AnimNextEditor/Private/Common/AnimNextFunctionItemDetails.cpp
+++ b/Engine/Plugins/Experimental/AnimNext/Source/AnimNextEditor/Private/Common/AnimNextFunctionItemDetails.cpp
@@ -106,7 +106,7 @@ void FAnimNextFunctionItemDetails::Delete(TConstArrayView<FWorkspaceOutlinerItem
 							if (URigVMNode* ModelNode = Model->FindNodeByName(*EdGraphNode->GetModelNodePath()))
 							{
 								FScopedTransaction Transaction(LOCTEXT("DeleteFunctionInOutliner", "Delete Function"));
-								EdGraphNode->GetController()->RemoveNode(ModelNode);
+								EdGraphNode->GetController()->RemoveNode(ModelNode, true, true);
 							}
 						}
 					}
@@ -122,7 +122,7 @@ void FAnimNextFunctionItemDetails::Delete(TConstArrayView<FWorkspaceOutlinerItem
 							{
 								FScopedTransaction Transaction(LOCTEXT("DeleteFunctionInOutliner", "Delete Function"));
 								URigVMLibraryNode* LibraryNode = CastChecked<URigVMLibraryNode>(EdGraph->GetModel()->GetOuter());
-								Controller->RemoveFunctionFromLibrary(LibraryNode->GetFName());
+								Controller->RemoveFunctionFromLibrary(LibraryNode->GetFName(), true, true);
 							}
 						}
 					}
diff --git a/Engine/Plugins/Experimental/AnimNext/Source/AnimNextEditor/Private/Common/GraphEditorSchemaActions.cpp b/Engine/Plugins/Experimental/AnimNext/Source/AnimNextEditor/Private/Common/GraphEditorSchemaActions.cpp
--- a/Engine/Plugins/Experimental/AnimNext/Source/AnimNextEditor/Private/Common/GraphEditorSchemaActions.cpp
+++ b/Engine/Plugins/Experimental/AnimNext/Source/AnimNextEditor/Private/Common/GraphEditorSchemaActions.cpp
@@ -54,7 +54,7 @@ UEdGraphNode* FAnimNextSchemaAction_RigUnit::PerformAction(UEdGraph* ParentGraph
 		FRigVMUnitNodeCreatedContext& UnitNodeCreatedContext = Controller->GetUnitNodeCreatedContext();
 		FRigVMUnitNodeCreatedContext::FScope ReasonScope(UnitNodeCreatedContext, ERigVMNodeCreatedReason::NodeSpawner, Host);
 
-		if (URigVMUnitNode* ModelNode = Controller->AddUnitNode(StructTemplate, NodeClass, FRigVMStruct::ExecuteName, Location, Name.ToString(), true, false))
+		if (URigVMUnitNode* ModelNode = Controller->AddUnitNode(StructTemplate, NodeClass, FRigVMStruct::ExecuteName, Location, Name.ToString(), true, true))
 		{
 			NewNode = Cast<URigVMEdGraphNode>(EdGraph->FindNodeForModelNodeName(ModelNode->GetFName()));
 			check(NewNode);
@@ -66,8 +66,8 @@ UEdGraphNode* FAnimNextSchemaAction_RigUnit::PerformAction(UEdGraph* ParentGraph
 					NewNode->AutowireNewNode(FromPin);
 				}
 
-				Controller->ClearNodeSelection(true);
-				Controller->SelectNode(ModelNode, true, true);
+				Controller->ClearNodeSelection(true, true);
+				Controller->SelectNode(ModelNode, true, true, true);
 			}
 		}
 		Controller->CloseUndoBracket();
@@ -118,7 +118,7 @@ UEdGraphNode* FAnimNextSchemaAction_DispatchFactory::PerformAction(UEdGraph* Par
 
 	Controller->OpenUndoBracket(FString::Printf(TEXT("Add '%s' Node"), *Name.ToString()));
 
-	if (URigVMTemplateNode* ModelNode = Controller->AddTemplateNode(Notation, Location, Name.ToString(), true, false))
+	if (URigVMTemplateNode* ModelNode = Controller->AddTemplateNode(Notation, Location, Name.ToString(), true, true))
 	{
 		NewNode = Cast<URigVMEdGraphNode>(EdGraph->FindNodeForModelNodeName(ModelNode->GetFName()));
 
@@ -129,8 +129,8 @@ UEdGraphNode* FAnimNextSchemaAction_DispatchFactory::PerformAction(UEdGraph* Par
 				NewNode->AutowireNewNode(FromPin);
 			}
 
-			Controller->ClearNodeSelection(true);
-			Controller->SelectNode(ModelNode, true, true);
+			Controller->ClearNodeSelection(true, true);
+			Controller->SelectNode(ModelNode, true, true, true);
 		}
 
 		Controller->CloseUndoBracket();
@@ -233,8 +233,8 @@ UEdGraphNode* FAnimNextSchemaAction_Variable::PerformAction(UEdGraph* ParentGrap
 
 			if (NewNode)
 			{
-				Controller->ClearNodeSelection(true);
-				Controller->SelectNode(ModelNode, true, true);
+				Controller->ClearNodeSelection(true, true);
+				Controller->SelectNode(ModelNode, true, true, true);
 			}
 			Controller->CloseUndoBracket();
 		}
@@ -377,8 +377,8 @@ UEdGraphNode* FAnimNextSchemaAction_Function::PerformAction(UEdGraph* ParentGrap
 
 			if (NewNode)
 			{
-				Controller->ClearNodeSelection(true);
-				Controller->SelectNode(ModelNode, true, true);
+				Controller->ClearNodeSelection(true, true);
+				Controller->SelectNode(ModelNode, true, true, true);
 			}
 			Controller->CloseUndoBracket();
 		}
diff --git a/Engine/Plugins/Experimental/AnimNext/Source/AnimNextUncookedOnly/Internal/AnimNextRigVMAssetEditorData.h b/Engine/Plugins/Experimental/AnimNext/Source/AnimNextUncookedOnly/Internal/AnimNextRigVMAssetEditorData.h
--- a/Engine/Plugins/Experimental/AnimNext/Source/AnimNextUncookedOnly/Internal/AnimNextRigVMAssetEditorData.h
+++ b/Engine/Plugins/Experimental/AnimNext/Source/AnimNextUncookedOnly/Internal/AnimNextRigVMAssetEditorData.h
@@ -124,6 +124,10 @@ public:
 	/** Adds a data interface to an AnimNext asset */
 	UFUNCTION(BlueprintCallable, Category = "AnimNext|Entries", meta=(ScriptMethod))
 	static ANIMNEXTUNCOOKEDONLY_API UAnimNextDataInterfaceEntry* AddDataInterface(UAnimNextRigVMAsset* InAsset, UAnimNextDataInterface* InDataInterface, bool bSetupUndoRedo = true, bool bPrintPythonCommand = true);
+
+	/** Adds a function to an AnimNext asset */
+	UFUNCTION(BlueprintCallable, Category = "AnimNext|Entries", meta=(ScriptMethod))
+	static ANIMNEXTUNCOOKEDONLY_API URigVMLibraryNode* AddFunction(UAnimNextRigVMAsset* InAsset, FName InFunctionName, bool bInMutable, bool bSetupUndoRedo = true, bool bPrintPythonCommand = true);
 };
 
 /* Base class for all AnimNext editor data objects that use RigVM */
@@ -143,7 +147,7 @@ public:
 	UAnimNextDataInterfaceEntry* AddDataInterface(UAnimNextDataInterface* InDataInterface, bool bSetupUndoRedo = true, bool bPrintPythonCommand = true);
 
 	/** Adds a function to this asset */
-	URigVMLibraryNode* AddFunction(FName InFunctionName, bool bInMutable, bool bInSetupUndoRedo = true, bool bInPrintPythonCommand = true);
+	URigVMLibraryNode* AddFunction(FName InFunctionName, bool bInMutable, bool bInSetupUndoRedo = true, bool bPrintPythonCommand = true);
 
 	// Find an entry by name
 	UAnimNextRigVMAssetEntry* FindEntry(FName InName) const;
diff --git a/Engine/Plugins/Experimental/AnimNext/Source/AnimNextUncookedOnly/Internal/UncookedOnlyUtils.h b/Engine/Plugins/Experimental/AnimNext/Source/AnimNextUncookedOnly/Internal/UncookedOnlyUtils.h
--- a/Engine/Plugins/Experimental/AnimNext/Source/AnimNextUncookedOnly/Internal/UncookedOnlyUtils.h
+++ b/Engine/Plugins/Experimental/AnimNext/Source/AnimNextUncookedOnly/Internal/UncookedOnlyUtils.h
@@ -198,7 +198,7 @@ struct ANIMNEXTUNCOOKEDONLY_API FUtils
 	static FRigVMTemplateArgumentType GetRigVMArgTypeFromParamType(const FAnimNextParamType& InParamType);
 
 	/** Set up a simple event graph */
-	static void SetupEventGraph(URigVMController* InController, UScriptStruct* InEventStruct, FName InEventName = NAME_None);
+	static void SetupEventGraph(URigVMController* InController, UScriptStruct* InEventStruct, FName InEventName = NAME_None, bool bPrintPythonCommand = false);
 
 	// Gets the variables that are exported to the asset registry for an asset
 	static bool GetExportedVariablesForAsset(const FAssetData& InAsset, FAnimNextAssetRegistryExports& OutExports);
diff --git a/Engine/Plugins/Experimental/AnimNext/Source/AnimNextUncookedOnly/Private/AnimNextRigVMAssetEditorData.cpp b/Engine/Plugins/Experimental/AnimNext/Source/AnimNextUncookedOnly/Private/AnimNextRigVMAssetEditorData.cpp
--- a/Engine/Plugins/Experimental/AnimNext/Source/AnimNextUncookedOnly/Private/AnimNextRigVMAssetEditorData.cpp
+++ b/Engine/Plugins/Experimental/AnimNext/Source/AnimNextUncookedOnly/Private/AnimNextRigVMAssetEditorData.cpp
@@ -948,7 +948,7 @@ void UAnimNextRigVMAssetEditorData::HandleRemoveNotify(UObject* InAsset, const F
 					if( (bFindWholeWord && EventNameString.Equals(InFindString, InSearchCase)) ||
 						(!bFindWholeWord && EventNameString.Contains(InFindString, InSearchCase)))
 					{
-						Controller->RemoveNode(Node);
+						Controller->RemoveNode(Node, true, true);
 					}
 				}
 			}
@@ -989,7 +989,7 @@ void UAnimNextRigVMAssetEditorData::HandleReplaceNotify(UObject* InAsset, const
 						(!bFindWholeWord && EventNameString.Contains(InFindString, InSearchCase)))
 					{
 						const FString NewName = EventNameString.Replace(*InFindString, *InReplaceString, InSearchCase);
-						Controller->SetPinDefaultValue(Pin, NewName, true, true, false);
+						Controller->SetPinDefaultValue(Pin->GetPinPath(), NewName, true, true, false, true);
 					}
 				}
 			}
@@ -1303,6 +1303,13 @@ bool UAnimNextRigVMAssetEditorData::RemoveEntry(UAnimNextRigVMAssetEntry* InEntr
 
 	BroadcastModified(EAnimNextEditorDataNotifType::EntryRemoved, this);
 
+	if (bPrintPythonCommand)
+	{
+		RigVMPythonUtils::Print(GetName(),
+				FString::Printf(TEXT("asset.remove_entry(asset.find_entry('%s'))"),
+				*InEntry->GetEntryName().ToString()));
+	}
+
 	return bResult;
 }
 
@@ -1319,12 +1326,32 @@ bool UAnimNextRigVMAssetEditorData::RemoveEntries(TConstArrayView<UAnimNextRigVM
 		TGuardValue<bool> DisableAutoCompile(bAutoRecompileVM, false);
 		for(UAnimNextRigVMAssetEntry* Entry : InEntries)
 		{
-			bResult |= RemoveEntry(Entry, bSetupUndoRedo, bPrintPythonCommand);
+			bResult |= RemoveEntry(Entry, bSetupUndoRedo, false);
 		}
 	}
 
 	BroadcastModified(EAnimNextEditorDataNotifType::EntryRemoved, this);
 
+	if (bPrintPythonCommand)
+	{
+
+		FString ArrayStr = TEXT("[");
+		for (int32 Index = 0; Index < InEntries.Num(); ++Index)
+		{
+			ArrayStr += TEXT("asset.find_entry('") + InEntries[Index]->GetEntryName().ToString() + TEXT("')");
+			if (Index < InEntries.Num() - 1)
+			{
+				ArrayStr += TEXT(", ");
+			}
+		}
+		ArrayStr += TEXT("]");
+
+
+		RigVMPythonUtils::Print(GetName(), 
+							FString::Printf(TEXT("asset.remove_entries(%s)"),
+											*ArrayStr));
+	}
+
 	return bResult;
 }
 
@@ -1342,12 +1369,19 @@ bool UAnimNextRigVMAssetEditorData::RemoveAllEntries(bool bSetupUndoRedo, bool b
 		TArray<UAnimNextRigVMAssetEntry*> EntriesCopy = Entries; 
 		for(UAnimNextRigVMAssetEntry* Entry : EntriesCopy)
 		{
-			bResult |= RemoveEntry(Entry, bSetupUndoRedo, bPrintPythonCommand);
+			bResult |= RemoveEntry(Entry, bSetupUndoRedo, false);
 		}
 	}
 
 	BroadcastModified(EAnimNextEditorDataNotifType::EntryRemoved, this);
 
+
+	if (bPrintPythonCommand)
+	{
+		RigVMPythonUtils::Print(GetName(), 
+							FString::Printf(TEXT("asset.remove_all_entries()")));
+	}
+
 	return bResult;
 }
 
@@ -1714,6 +1748,20 @@ UAnimNextVariableEntry* UAnimNextRigVMAssetEditorData::AddVariable(FName InName,
 
 	BroadcastModified(EAnimNextEditorDataNotifType::EntryAdded, NewEntry);
 
+	if (bPrintPythonCommand)
+	{
+		const FString ValueTypeString = InType.GetValueTypeObject() ?
+				FString::Printf(TEXT("unreal.%s.static_%s()"), *InType.GetValueTypeObject()->GetName(), InType.GetValueTypeObject()->IsA<UScriptStruct>() ? TEXT("struct") : TEXT("class"))
+				: TEXT("None");
+		RigVMPythonUtils::Print(GetName(), 
+							FString::Printf(TEXT("asset.add_variable('%s', %s, %s, %s, '%s')"),
+											*InName.ToString(),
+											*RigVMPythonUtils::EnumValueToPythonString<EPropertyBagPropertyType>(static_cast<int64>(InType.GetValueType())),
+											*RigVMPythonUtils::EnumValueToPythonString<EPropertyBagContainerType>(static_cast<int64>(InType.GetContainerType())),
+											*ValueTypeString,
+											*InDefaultValue));
+	}
+
 	return NewEntry;
 }
 
@@ -1796,6 +1844,13 @@ UAnimNextEventGraphEntry* UAnimNextRigVMAssetEditorData::AddEventGraph(FName InN
 
 	BroadcastModified(EAnimNextEditorDataNotifType::EntryAdded, NewEntry);
 
+	if (bPrintPythonCommand)
+	{
+		RigVMPythonUtils::Print(GetName(),
+				FString::Printf(TEXT("asset.add_event_graph('%s', unreal.%s)"),
+				*InName.ToString(), *InEventStruct->GetName()));
+	}
+
 	return NewEntry;
 }
 
@@ -1805,6 +1860,11 @@ UAnimNextDataInterfaceEntry* UAnimNextRigVMAssetLibrary::AddDataInterface(UAnimN
 	return UE::AnimNext::UncookedOnly::FUtils::GetEditorData(InAsset)->AddDataInterface(InDataInterface, bSetupUndoRedo, bPrintPythonCommand);
 }
 
+URigVMLibraryNode* UAnimNextRigVMAssetLibrary::AddFunction(UAnimNextRigVMAsset* InAsset, FName InFunctionName, bool bInMutable, bool bSetupUndoRedo, bool bPrintPythonCommand)
+{
+	return UE::AnimNext::UncookedOnly::FUtils::GetEditorData(InAsset)->AddFunction(InFunctionName, bInMutable, bSetupUndoRedo, bPrintPythonCommand);
+}
+
 UAnimNextDataInterfaceEntry* UAnimNextRigVMAssetEditorData::AddDataInterface(UAnimNextDataInterface* InDataInterface, bool bSetupUndoRedo, bool bPrintPythonCommand)
 {
 	if(InDataInterface == nullptr)
@@ -1920,13 +1980,31 @@ UAnimNextDataInterfaceEntry* UAnimNextRigVMAssetEditorData::AddDataInterface(UAn
 
 	BroadcastModified(EAnimNextEditorDataNotifType::EntryAdded, NewEntry);
 
+	if (bPrintPythonCommand)
+	{
+		RigVMPythonUtils::Print(GetName(),
+				FString::Printf(TEXT("asset.add_data_interface(unreal.find_object(outer=None, name='%s'))"),
+				 *InDataInterface->GetPathName()));
+	}
+	
+
 	return NewEntry;
 }
 
-URigVMLibraryNode* UAnimNextRigVMAssetEditorData::AddFunction(FName InFunctionName, bool bInMutable, bool bInSetupUndoRedo, bool bInPrintPythonCommand)
+URigVMLibraryNode* UAnimNextRigVMAssetEditorData::AddFunction(FName InFunctionName, bool bInMutable, bool bInSetupUndoRedo, bool bPrintPythonCommand)
 {
 	URigVMController* Controller = RigVMClient.GetOrCreateController(GetLocalFunctionLibrary());
-	return Controller->AddFunctionToLibrary(InFunctionName, bInMutable, FVector2D::ZeroVector, bInSetupUndoRedo, bInPrintPythonCommand);
+	URigVMLibraryNode* Node = Controller->AddFunctionToLibrary(InFunctionName, bInMutable, FVector2D::ZeroVector, bInSetupUndoRedo, false);
+
+	if (bPrintPythonCommand)
+	{
+		RigVMPythonUtils::Print(GetName(),
+				FString::Printf(TEXT("asset.add_function('%s', %s)"),
+				 *InFunctionName.ToString(),
+				 bInMutable ? TEXT("True") : TEXT("False")));
+	}
+
+	return Node;
 }
 
 bool UAnimNextRigVMAssetEditorData::HasPublicVariables() const
diff --git a/Engine/Plugins/Experimental/AnimNext/Source/AnimNextUncookedOnly/Private/UncookedOnlyUtils.cpp b/Engine/Plugins/Experimental/AnimNext/Source/AnimNextUncookedOnly/Private/UncookedOnlyUtils.cpp
--- a/Engine/Plugins/Experimental/AnimNext/Source/AnimNextUncookedOnly/Private/UncookedOnlyUtils.cpp
+++ b/Engine/Plugins/Experimental/AnimNext/Source/AnimNextUncookedOnly/Private/UncookedOnlyUtils.cpp
@@ -790,18 +790,18 @@ FRigVMTemplateArgumentType FUtils::GetRigVMArgTypeFromParamType(const FAnimNextP
 	return ArgType;
 }
 
-void FUtils::SetupEventGraph(URigVMController* InController, UScriptStruct* InEventStruct, FName InEventName)
+void FUtils::SetupEventGraph(URigVMController* InController, UScriptStruct* InEventStruct, FName InEventName, bool bPrintPythonCommand)
 {
 	// Clear the graph
 	InController->RemoveNodes(InController->GetGraph()->GetNodes());
 
 	// Add entry point
-	URigVMUnitNode* EventNode = InController->AddUnitNode(InEventStruct, FRigVMStruct::ExecuteName, FVector2D(-200.0f, 0.0f), FString(), false);
+	URigVMUnitNode* EventNode = InController->AddUnitNode(InEventStruct, FRigVMStruct::ExecuteName, FVector2D(-200.0f, 0.0f), FString(), false, bPrintPythonCommand);
 
 	// Setup custom event name if required
 	if(EventNode && InEventStruct == FRigVMFunction_UserDefinedEvent::StaticStruct() && InEventName != NAME_None && !InController->GetAllEventNames().Contains(InEventName))
 	{
-		InController->SetPinDefaultValue(EventNode->FindPin(GET_MEMBER_NAME_STRING_CHECKED(FRigVMFunction_UserDefinedEvent, EventName)), InEventName.ToString(), true, true, false);
+		InController->SetPinDefaultValue(EventNode->FindPin(GET_MEMBER_NAME_STRING_CHECKED(FRigVMFunction_UserDefinedEvent, EventName)), InEventName.ToString(), true, true, bPrintPythonCommand, false);
 	}
 }
 
diff --git a/Engine/Plugins/Experimental/AnimNextAnimGraph/Source/AnimNextAnimGraphEditor/Private/Graph/AnimNextAnimationGraphFactory.cpp b/Engine/Plugins/Experimental/AnimNextAnimGraph/Source/AnimNextAnimGraphEditor/Private/Graph/AnimNextAnimationGraphFactory.cpp
--- a/Engine/Plugins/Experimental/AnimNextAnimGraph/Source/AnimNextAnimGraphEditor/Private/Graph/AnimNextAnimationGraphFactory.cpp
+++ b/Engine/Plugins/Experimental/AnimNextAnimGraph/Source/AnimNextAnimGraphEditor/Private/Graph/AnimNextAnimationGraphFactory.cpp
@@ -36,7 +36,7 @@ UObject* UAnimNextAnimationGraphFactory::FactoryCreateNew(UClass* Class, UObject
 	EditorData->Initialize(/*bRecompileVM*/false);
 
 	// Add a single internal graph entry
-	UAnimNextAnimationGraphEntry* NewGraph = EditorData->AddAnimationGraph(FRigUnit_AnimNextGraphRoot::DefaultEntryPoint);
+	UAnimNextAnimationGraphEntry* NewGraph = EditorData->AddAnimationGraph(FRigUnit_AnimNextGraphRoot::DefaultEntryPoint, false, false);
 
 	// Hide in the outliner - the asset itself represents the graph to users
 	NewGraph->SetHiddenInOutliner(true);
diff --git a/Engine/Plugins/Experimental/AnimNextAnimGraph/Source/AnimNextAnimGraphEditor/Private/Graph/AnimNextAnimationGraphMenuExtensions.cpp b/Engine/Plugins/Experimental/AnimNextAnimGraph/Source/AnimNextAnimGraphEditor/Private/Graph/AnimNextAnimationGraphMenuExtensions.cpp
--- a/Engine/Plugins/Experimental/AnimNextAnimGraph/Source/AnimNextAnimGraphEditor/Private/Graph/AnimNextAnimationGraphMenuExtensions.cpp
+++ b/Engine/Plugins/Experimental/AnimNextAnimGraph/Source/AnimNextAnimGraphEditor/Private/Graph/AnimNextAnimationGraphMenuExtensions.cpp
@@ -153,7 +153,7 @@ void FAnimationGraphMenuExtensions::RegisterMenus()
 					FUIAction(
 						FExecuteAction::CreateLambda([VMController, ModelNode, Name = VMPin->GetFName()]()
 						{
-							VMController->RemoveTrait(ModelNode, Name);
+							VMController->RemoveTrait(ModelNode->GetFName(), Name, true, true);
 						})
 					));
 			}
@@ -193,7 +193,7 @@ void FAnimationGraphMenuExtensions::RegisterMenus()
 							return ModelNode->FindPin(InName.ToString()) == nullptr;
 						}, false, false);
 
-						VMController->AddTrait(ModelNode, FRigVMTrait_AnimNextPublicVariables::StaticStruct(), ValidTraitName, DefaultValue);
+						VMController->AddTrait(ModelNode->GetFName(), *FRigVMTrait_AnimNextPublicVariables::StaticStruct()->GetPathName(), ValidTraitName, DefaultValue, INDEX_NONE, true, true);
 					});
 
 					AssetPickerConfig.OnShouldFilterAsset = FOnShouldFilterAsset::CreateLambda([](const FAssetData& InAssetData)
@@ -279,7 +279,7 @@ void FAnimationGraphMenuExtensions::AddManifestSection(UToolMenu* InMenu, const
 							{
 								FRigVMControllerCompileBracketScope CompileScope(Controller);
 
-								Controller->AddNodeToManifest(ModelNode);
+								Controller->AddNodeToManifest(ModelNode, true, true);
 								
 								if (UAnimNextRigVMAssetEditorData* EditorData = EdGraphNode->GetGraph()->GetTypedOuter<UAnimNextRigVMAssetEditorData>())
 								{
@@ -306,7 +306,7 @@ void FAnimationGraphMenuExtensions::AddManifestSection(UToolMenu* InMenu, const
 
 						FRigVMControllerCompileBracketScope CompileScope(Controller);
 
-						Controller->RemoveNodeFromManifest(EdGraphNode->GetModelNode());
+						Controller->RemoveNodeFromManifest(EdGraphNode->GetModelNode(), true, true);
 						
 						if (UAnimNextRigVMAssetEditorData* EditorData = EdGraphNode->GetGraph()->GetTypedOuter<UAnimNextRigVMAssetEditorData>())
 						{
diff --git a/Engine/Plugins/Experimental/AnimNextAnimGraph/Source/AnimNextAnimGraphEditor/Private/Graph/STraitEditorView.cpp b/Engine/Plugins/Experimental/AnimNextAnimGraph/Source/AnimNextAnimGraphEditor/Private/Graph/STraitEditorView.cpp
--- a/Engine/Plugins/Experimental/AnimNextAnimGraph/Source/AnimNextAnimGraphEditor/Private/Graph/STraitEditorView.cpp
+++ b/Engine/Plugins/Experimental/AnimNextAnimGraph/Source/AnimNextAnimGraphEditor/Private/Graph/STraitEditorView.cpp
@@ -126,11 +126,11 @@ FReply STraitEditorView::OnTraitClicked(const FTraitUID InClickedTraitUID)
 				const FName NewTraitTypeName = *Trait->GetTraitName();
 				if (TraitIndex == INDEX_NONE || bIsAddingMissingBaseTrait)
 				{
-					Controller->AddTraitByName(EdGraphNode->GetFName(), NewTraitTypeName, PinIndex);
+					Controller->AddTraitByName(EdGraphNode->GetFName(), NewTraitTypeName, PinIndex,TEXT(""), true, true);
 				}
 				else
 				{
-					Controller->SwapTraitByName(EdGraphNode->GetFName(), SwapTraitData->TraitName, PinIndex, NewTraitTypeName);
+					Controller->SwapTraitByName(EdGraphNode->GetFName(), SwapTraitData->TraitName, PinIndex, NewTraitTypeName, TEXT(""), true, true);
 				}
 			}
 		}
@@ -162,7 +162,7 @@ FReply STraitEditorView::OnTraitDeleteRequest(const FTraitUID InTraitUIDToDelete
 					{
 						if (UAnimNextController* Controller = Cast<UAnimNextController>(EdGraphNode->GetController()))
 						{
-							Controller->RemoveTraitByName(EdGraphNode->GetFName(), TraitDataEditorDef->TraitName);
+							Controller->RemoveTraitByName(EdGraphNode->GetFName(), TraitDataEditorDef->TraitName, true, true);
 						}
 					}
 				}
@@ -210,11 +210,11 @@ void STraitEditorView::ExecuteTraitDrag(const FTraitUID DraggedTraitUID, const F
 					{
 						if(bIsAddingMissingBaseTrait)
 						{
-							Controller->AddTraitByName(EdGraphNode->GetFName(), DraggedTraitTypeName, PinIndex);
+							Controller->AddTraitByName(EdGraphNode->GetFName(), DraggedTraitTypeName, PinIndex, TEXT(""), true, true);
 						}
 						else
 						{
-							Controller->SwapTraitByName(EdGraphNode->GetFName(), SwapTraitData->TraitName, PinIndex, DraggedTraitTypeName);
+							Controller->SwapTraitByName(EdGraphNode->GetFName(), SwapTraitData->TraitName, PinIndex, DraggedTraitTypeName, TEXT(""), true, true);
 						}
 					}
 					else
@@ -224,11 +224,11 @@ void STraitEditorView::ExecuteTraitDrag(const FTraitUID DraggedTraitUID, const F
 
 						if (DraggedTraitIndex != INDEX_NONE)
 						{
-							Controller->SetTraitPinIndex(EdGraphNode->GetFName(), DraggedTraitData->TraitName, PinIndex + 1);
+							Controller->SetTraitPinIndex(EdGraphNode->GetFName(), DraggedTraitData->TraitName, PinIndex + 1, true, true);
 						}
 						else
 						{
-							Controller->AddTraitByName(EdGraphNode->GetFName(), DraggedTraitTypeName, PinIndex + 1);
+							Controller->AddTraitByName(EdGraphNode->GetFName(), DraggedTraitTypeName, PinIndex + 1, TEXT(""), true, true);
 						}
 					}
 				}
diff --git a/Engine/Plugins/Experimental/AnimNextAnimGraph/Source/AnimNextAnimGraphEditor/Private/Graph/STraitStackView.cpp b/Engine/Plugins/Experimental/AnimNextAnimGraph/Source/AnimNextAnimGraphEditor/Private/Graph/STraitStackView.cpp
--- a/Engine/Plugins/Experimental/AnimNextAnimGraph/Source/AnimNextAnimGraphEditor/Private/Graph/STraitStackView.cpp
+++ b/Engine/Plugins/Experimental/AnimNextAnimGraph/Source/AnimNextAnimGraphEditor/Private/Graph/STraitStackView.cpp
@@ -145,7 +145,7 @@ void STraitStackView::Construct(const FArguments& InArgs, TSharedPtr<FTraitEdito
 													{
 														if (URigVMController* Controller = EdGraphNode->GetController())
 														{
-															if(Controller->SetNodeTitle(ModelNode, InNewText.ToString()))
+															if(Controller->SetNodeTitle(ModelNode, InNewText.ToString(), true, false, true))
 															{
 																if (UnitNode->IsExposedToManifest()) // if the node is exposed to manifest, we have to force a recompile to update asset tags
 																{
diff --git a/Engine/Plugins/Experimental/AnimNextAnimGraph/Source/AnimNextAnimGraphEditor/Private/Traits/AnimNextCallFunctionSharedDataDetails.cpp b/Engine/Plugins/Experimental/AnimNextAnimGraph/Source/AnimNextAnimGraphEditor/Private/Traits/AnimNextCallFunctionSharedDataDetails.cpp
--- a/Engine/Plugins/Experimental/AnimNextAnimGraph/Source/AnimNextAnimGraphEditor/Private/Traits/AnimNextCallFunctionSharedDataDetails.cpp
+++ b/Engine/Plugins/Experimental/AnimNextAnimGraph/Source/AnimNextAnimGraphEditor/Private/Traits/AnimNextCallFunctionSharedDataDetails.cpp
@@ -185,7 +185,7 @@ void FCallFunctionSharedDataDetails::CustomizeChildren(TSharedRef<IPropertyHandl
 			{
 				FRigVMControllerCompileBracketScope CompileScope(Controller);
 				FScopedTransaction Transaction(LOCTEXT("AddFunctionTransaction", "Add Function"));
-				NewFunction = EditorData->AddFunction(TEXT("NewFunction"), true);
+				NewFunction = EditorData->AddFunction(TEXT("NewFunction"), true, true, true);
 
 				if(NewFunction == nullptr)
 				{
diff --git a/Engine/Plugins/Experimental/AnimNextAnimGraph/Source/AnimNextAnimGraphUncookedOnly/Internal/AnimGraphUncookedOnlyUtils.h b/Engine/Plugins/Experimental/AnimNextAnimGraph/Source/AnimNextAnimGraphUncookedOnly/Internal/AnimGraphUncookedOnlyUtils.h
--- a/Engine/Plugins/Experimental/AnimNextAnimGraph/Source/AnimNextAnimGraphUncookedOnly/Internal/AnimGraphUncookedOnlyUtils.h
+++ b/Engine/Plugins/Experimental/AnimNextAnimGraph/Source/AnimNextAnimGraphUncookedOnly/Internal/AnimGraphUncookedOnlyUtils.h
@@ -16,7 +16,7 @@ namespace UE::AnimNext::UncookedOnly
 struct ANIMNEXTANIMGRAPHUNCOOKEDONLY_API FAnimGraphUtils
 {
 	/** Set up a simple animation graph */
-	static void SetupAnimGraph(const FName EntryName, URigVMController* InController, bool bSetupUndoRedo = true);
+	static void SetupAnimGraph(const FName EntryName, URigVMController* InController, bool bSetupUndoRedo = true, bool bPrintPythonCommand = false);
 
 	/** Check whether the supplied node is a trait stack */
 	static bool IsTraitStackNode(const URigVMNode* InModelNode);
diff --git a/Engine/Plugins/Experimental/AnimNextAnimGraph/Source/AnimNextAnimGraphUncookedOnly/Internal/Graph/AnimGraphEditorSchemaActions.cpp b/Engine/Plugins/Experimental/AnimNextAnimGraph/Source/AnimNextAnimGraphUncookedOnly/Internal/Graph/AnimGraphEditorSchemaActions.cpp
--- a/Engine/Plugins/Experimental/AnimNextAnimGraph/Source/AnimNextAnimGraphUncookedOnly/Internal/Graph/AnimGraphEditorSchemaActions.cpp
+++ b/Engine/Plugins/Experimental/AnimNextAnimGraph/Source/AnimNextAnimGraphUncookedOnly/Internal/Graph/AnimGraphEditorSchemaActions.cpp
@@ -72,7 +72,7 @@ UEdGraphNode* FAnimNextSchemaAction_AddManifestNode::PerformAction(UEdGraph* Par
 
 			const FString NodeData = MoveTemp(Archive);
 
-			TArray<FName> NodeNamesCreated = UE::RigVM::Editor::Tools::ImportNodesFromText(Location, NodeData, Controller, EdGraph->GetModel(), Host->GetLocalFunctionLibrary(), Host->GetRigVMGraphFunctionHost());
+			TArray<FName> NodeNamesCreated = UE::RigVM::Editor::Tools::ImportNodesFromText(Location, NodeData, Controller, EdGraph->GetModel(), Host->GetLocalFunctionLibrary(), Host->GetRigVMGraphFunctionHost(), true, true);
 
 			TArray<URigVMNode*> ModelNodes;
 			for (const FName& Name : NodeNamesCreated)
@@ -91,7 +91,7 @@ UEdGraphNode* FAnimNextSchemaAction_AddManifestNode::PerformAction(UEdGraph* Par
 
 				if (NewNode)
 				{
-					Controller->RemoveNodeFromManifest(ModelNode, false, false);
+					Controller->RemoveNodeFromManifest(ModelNode, false, true);
 				}
 				Controller->CloseUndoBracket();
 			}
@@ -128,7 +128,7 @@ UEdGraphNode* FAnimNextSchemaAction_NotifyEvent::PerformAction(UEdGraph* ParentG
 		{
 			Controller->OpenUndoBracket(LOCTEXT("AddNotifyEventTransaction", "Add Notify Event").ToString());
 			URigVMUnitNode* CustomEventNode = Controller->AddUnitNode(FRigVMFunction_UserDefinedEvent::StaticStruct(), FRigVMStruct::ExecuteName, Location);
-			Controller->SetPinDefaultValue(CustomEventNode->FindPin(GET_MEMBER_NAME_STRING_CHECKED(FRigVMFunction_UserDefinedEvent, EventName)), InNotifyEventName.ToString(), true, true, true);
+			Controller->SetPinDefaultValue(CustomEventNode->FindPin(GET_MEMBER_NAME_STRING_CHECKED(FRigVMFunction_UserDefinedEvent, EventName))->GetPinPath(), InNotifyEventName.ToString(), true, true, true, true, true);
 			Controller->CloseUndoBracket();
 		}
 	};
diff --git a/Engine/Plugins/Experimental/AnimNextAnimGraph/Source/AnimNextAnimGraphUncookedOnly/Private/AnimGraphUncookedOnlyUtils.cpp b/Engine/Plugins/Experimental/AnimNextAnimGraph/Source/AnimNextAnimGraphUncookedOnly/Private/AnimGraphUncookedOnlyUtils.cpp
--- a/Engine/Plugins/Experimental/AnimNextAnimGraph/Source/AnimNextAnimGraphUncookedOnly/Private/AnimGraphUncookedOnlyUtils.cpp
+++ b/Engine/Plugins/Experimental/AnimNextAnimGraph/Source/AnimNextAnimGraphUncookedOnly/Private/AnimGraphUncookedOnlyUtils.cpp
@@ -27,13 +27,13 @@ bool FAnimGraphUtils::IsTraitStackNode(const URigVMNode* InModelNode)
 	return false;
 }
 
-void FAnimGraphUtils::SetupAnimGraph(const FName EntryName, URigVMController* InController, bool bSetupUndoRedo /*= true*/)
+void FAnimGraphUtils::SetupAnimGraph(const FName EntryName, URigVMController* InController, bool bSetupUndoRedo /*= true*/,  bool bPrintPythonCommand /*= false*/ )
 {
 	// Clear the graph
-	InController->RemoveNodes(InController->GetGraph()->GetNodes());
+	InController->RemoveNodes(InController->GetGraph()->GetNodes(), bSetupUndoRedo, bPrintPythonCommand);
 
 	// Add root node
-	URigVMUnitNode* MainEntryPointNode = InController->AddUnitNode(FRigUnit_AnimNextGraphRoot::StaticStruct(), FRigUnit_AnimNextGraphRoot::EventName, FVector2D(-400.0f, 0.0f), FString(), bSetupUndoRedo);
+	URigVMUnitNode* MainEntryPointNode = InController->AddUnitNode(FRigUnit_AnimNextGraphRoot::StaticStruct(), FRigUnit_AnimNextGraphRoot::EventName, FVector2D(-400.0f, 0.0f), FString(), bSetupUndoRedo, bPrintPythonCommand);
 	if(MainEntryPointNode == nullptr)
 	{
 		return;
@@ -51,7 +51,7 @@ void FAnimGraphUtils::SetupAnimGraph(const FName EntryName, URigVMController* In
 		return;
 	}
 
-	InController->SetPinDefaultValue(EntryPointPin->GetPinPath(), EntryName.ToString(), true, bSetupUndoRedo);
+	InController->SetPinDefaultValue(EntryPointPin->GetPinPath(), EntryName.ToString(), true, bSetupUndoRedo, true, bPrintPythonCommand);
 }
 
 void FAnimGraphUtils::GetAssetManifestNodesRegistryExports(const UAnimNextRigVMAssetEditorData* InEditorData, FAnimNextAssetRegistryExports& OutExports)
diff --git a/Engine/Plugins/Experimental/AnimNextAnimGraph/Source/AnimNextAnimGraphUncookedOnly/Private/Entries/AnimNextAnimationGraphEntry.cpp b/Engine/Plugins/Experimental/AnimNextAnimGraph/Source/AnimNextAnimGraphUncookedOnly/Private/Entries/AnimNextAnimationGraphEntry.cpp
--- a/Engine/Plugins/Experimental/AnimNextAnimGraph/Source/AnimNextAnimGraphUncookedOnly/Private/Entries/AnimNextAnimationGraphEntry.cpp
+++ b/Engine/Plugins/Experimental/AnimNextAnimGraph/Source/AnimNextAnimGraphUncookedOnly/Private/Entries/AnimNextAnimationGraphEntry.cpp
@@ -65,7 +65,7 @@ void UAnimNextAnimationGraphEntry::SetEntryName(FName InName, bool bSetupUndoRed
 				check(EntryPointPin);
 				check(EntryPointPin->GetDirection() == ERigVMPinDirection::Hidden);
 
-				Controller->SetPinDefaultValue(EntryPointPin->GetPinPath(), InName.ToString());
+				Controller->SetPinDefaultValue(EntryPointPin->GetPinPath(), InName.ToString(), true, true, false, true);
 			}
 		}
 	}
diff --git a/Engine/Plugins/Experimental/AnimNextAnimGraph/Source/AnimNextAnimGraphUncookedOnly/Private/Graph/AnimNextAnimationGraph_EditorData.cpp b/Engine/Plugins/Experimental/AnimNextAnimGraph/Source/AnimNextAnimGraphUncookedOnly/Private/Graph/AnimNextAnimationGraph_EditorData.cpp
--- a/Engine/Plugins/Experimental/AnimNextAnimGraph/Source/AnimNextAnimGraphUncookedOnly/Private/Graph/AnimNextAnimationGraph_EditorData.cpp
+++ b/Engine/Plugins/Experimental/AnimNextAnimGraph/Source/AnimNextAnimGraphUncookedOnly/Private/Graph/AnimNextAnimationGraph_EditorData.cpp
@@ -4,6 +4,7 @@
 #include "AnimGraphUncookedOnlyUtils.h"
 #include "AnimNextAnimGraphWorkspaceAssetUserData.h"
 #include "AnimNextScopedCompilerResults.h"
+#include "RigVMPythonUtils.h"
 #include "Compilation/AnimNextGetFunctionHeaderCompileContext.h"
 #include "Compilation/AnimNextGetVariableCompileContext.h"
 #include "Compilation/AnimNextGetGraphCompileContext.h"
@@ -863,6 +864,13 @@ UAnimNextAnimationGraphEntry* UAnimNextAnimationGraph_EditorData::AddAnimationGr
 
 	BroadcastModified(EAnimNextEditorDataNotifType::EntryAdded, NewEntry);
 
+	if (bPrintPythonCommand)
+	{
+		RigVMPythonUtils::Print(GetName(),
+				FString::Printf(TEXT("asset.add_animation_graph('%s')"),
+				*InName.ToString()));
+	}
+
 	return NewEntry;
 }
 
diff --git a/Engine/Plugins/Runtime/RigVM/Source/RigVMDeveloper/Private/RigVMModel/RigVMController.cpp b/Engine/Plugins/Runtime/RigVM/Source/RigVMDeveloper/Private/RigVMModel/RigVMController.cpp
--- a/Engine/Plugins/Runtime/RigVM/Source/RigVMDeveloper/Private/RigVMModel/RigVMController.cpp
+++ b/Engine/Plugins/Runtime/RigVM/Source/RigVMDeveloper/Private/RigVMModel/RigVMController.cpp
@@ -13948,7 +13948,7 @@ URigVMLibraryNode* URigVMController::AddFunctionToLibrary(const FName& InFunctio
 	return CollapseNode;
 }
 
-bool URigVMController::RemoveFunctionFromLibrary(const FName& InFunctionName, bool bSetupUndoRedo)
+bool URigVMController::RemoveFunctionFromLibrary(const FName& InFunctionName, bool bSetupUndoRedo, bool bPrintPythonCommand)
 {
 	if (!IsValidGraph())
 	{
@@ -13969,14 +13969,14 @@ bool URigVMController::RemoveFunctionFromLibrary(const FName& InFunctionName, bo
 		{
 			return false;
 		}
-		return RemoveNodeByName(InFunctionName, bSetupUndoRedo);
+		return RemoveNodeByName(InFunctionName, bSetupUndoRedo, bPrintPythonCommand);
 	}
 
 	ReportErrorf(TEXT("Cannot find node '%s'."), *InFunctionName.ToString());
 	return false;
 }
 
-bool URigVMController::RenameFunction(const FName& InOldFunctionName, const FName& InNewFunctionName, bool bSetupUndoRedo)
+bool URigVMController::RenameFunction(const FName& InOldFunctionName, const FName& InNewFunctionName, bool bSetupUndoRedo, bool bPrintPythonCommand)
 {
 	if (!IsValidGraph())
 	{
@@ -14005,7 +14005,7 @@ bool URigVMController::RenameFunction(const FName& InOldFunctionName, const FNam
 		return false;
 	}
 
-	return RenameNode(Node, InNewFunctionName, bSetupUndoRedo);
+	return RenameNode(Node, InNewFunctionName, bSetupUndoRedo, bPrintPythonCommand);
 }
 
 bool URigVMController::MarkFunctionAsPublic(const FName& InFunctionName, bool bInIsPublic, bool bSetupUndoRedo,	bool bPrintPythonCommand)
diff --git a/Engine/Plugins/Runtime/RigVM/Source/RigVMDeveloper/Public/RigVMModel/RigVMController.h b/Engine/Plugins/Runtime/RigVM/Source/RigVMDeveloper/Public/RigVMModel/RigVMController.h
--- a/Engine/Plugins/Runtime/RigVM/Source/RigVMDeveloper/Public/RigVMModel/RigVMController.h
+++ b/Engine/Plugins/Runtime/RigVM/Source/RigVMDeveloper/Public/RigVMModel/RigVMController.h
@@ -1032,11 +1032,11 @@ public:
 
 	// Removes a function from a function library graph
 	UFUNCTION(BlueprintCallable, Category = RigVMController)
-	bool RemoveFunctionFromLibrary(const FName& InFunctionName, bool bSetupUndoRedo = true);
+	bool RemoveFunctionFromLibrary(const FName& InFunctionName, bool bSetupUndoRedo = true, bool bPrintPythonCommand = false);
 
 	// Renames a function in the function library
 	UFUNCTION(BlueprintCallable, Category = RigVMController)
-	bool RenameFunction(const FName& InOldFunctionName, const FName& InNewFunctionName, bool bSetupUndoRedo = true);
+	bool RenameFunction(const FName& InOldFunctionName, const FName& InNewFunctionName, bool bSetupUndoRedo = true, bool bPrintPythonCommand = false);
 
 	// Mark a function as public/private in the function library
 	UFUNCTION(BlueprintCallable, Category = RigVMController)
diff --git a/Engine/Plugins/Runtime/RigVM/Source/RigVMEditor/Private/Editor/RigVMEditorTools.cpp b/Engine/Plugins/Runtime/RigVM/Source/RigVMEditor/Private/Editor/RigVMEditorTools.cpp
--- a/Engine/Plugins/Runtime/RigVM/Source/RigVMEditor/Private/Editor/RigVMEditorTools.cpp
+++ b/Engine/Plugins/Runtime/RigVM/Source/RigVMEditor/Private/Editor/RigVMEditorTools.cpp
@@ -28,7 +28,7 @@ RIGVMEDITOR_API bool PasteNodes(const FVector2D& PasteLocation
 		InFocusedController->OpenUndoBracket(TEXT("Paste Nodes."));
 	}
 
-	TArray<FName> NodeNamesCreated = ImportNodesFromText(PasteLocation, TextToImport, InFocusedController, InFocusedModel, InLocalFunctionLibrary, InGraphFunctionHost);
+	TArray<FName> NodeNamesCreated = ImportNodesFromText(PasteLocation, TextToImport, InFocusedController, InFocusedModel, InLocalFunctionLibrary, InGraphFunctionHost, bPrintPythonCommands);
 
 	const bool bPastePerformed = NodeNamesCreated.Num() > 0;
 	if (bPastePerformed)
