Anim Next: Fix for enum properties using SetVariable in blueprint since most pins are implicitly converted to FByteProperty

diff --git a/Engine/Plugins/Experimental/AnimNext/Source/AnimNext/Private/Component/AnimNextComponent.cpp b/Engine/Plugins/Experimental/AnimNext/Source/AnimNext/Private/Component/AnimNextComponent.cpp
--- a/Engine/Plugins/Experimental/AnimNext/Source/AnimNext/Private/Component/AnimNextComponent.cpp
+++ b/Engine/Plugins/Experimental/AnimNext/Source/AnimNext/Private/Component/AnimNextComponent.cpp
@@ -163,6 +163,18 @@ DEFINE_FUNCTION(UAnimNextComponent::execSetVariable)
 		const FProperty* CachedProperty = ProxyDescs[*IndexPtr].CachedProperty;
 
 		bIsTypeMatching = CachedProperty->GetClass() == ValueProp->GetClass();
+
+		if (!bIsTypeMatching)
+		{
+			if (const FByteProperty* SourceByteProperty = CastField<FByteProperty>(ValueProp))
+			{
+				if (const FEnumProperty* TargetEnumProperty = CastField<FEnumProperty>(CachedProperty))
+				{
+					bIsTypeMatching = TargetEnumProperty->GetUnderlyingProperty()->IsA<FByteProperty>() || SourceByteProperty->Enum == TargetEnumProperty->GetEnum();
+				}
+			}
+		}
+
 		if (bIsTypeMatching)
 		{
 			const void* ValuePtr = ValueProp->ContainerPtrToValuePtr<void>(ContainerPtr);
