AnimNext: Add custom module events & reduce tick functions

Adds the ability for multiple module events per-phase (e.g 2 PrePhysics events that can be independently depended upon by others). Events are all sorted strictly across the frame, and all depend on one another in a linear chain to avoid racing on module data.
Removed the 'initialize' and 'end' tick functions and replaced them with branches and a more explicit separation of initialization state and run state.
Event setup is now more tightly bound to tick function semantics, so events now specify their tick group directly, rather than relying on the binding function.
Removed old GetOptions-style picking of pre-defined module events in place of a custom widget for picking. This allows choosing from default events and custom events from users, discovered via bytecode introspection.

diff --git a/Engine/Plugins/Experimental/AnimNext/Source/AnimNext/Internal/AnimNextConfig.h b/Engine/Plugins/Experimental/AnimNext/Source/AnimNext/Internal/AnimNextConfig.h
--- a/Engine/Plugins/Experimental/AnimNext/Source/AnimNext/Internal/AnimNextConfig.h
+++ b/Engine/Plugins/Experimental/AnimNext/Source/AnimNext/Internal/AnimNextConfig.h
@@ -14,8 +14,4 @@ private:
 #if WITH_EDITOR
 	virtual void PostEditChangeProperty(FPropertyChangedEvent& PropertyChangedEvent) override;
 #endif
-
-	// Helper function used to display UI options for module events
-	UFUNCTION()
-	static TArray<FString> GetModuleEventNames();
 };
\ No newline at end of file
diff --git a/Engine/Plugins/Experimental/AnimNext/Source/AnimNext/Internal/Component/AnimNextWorldSubsystem.h b/Engine/Plugins/Experimental/AnimNext/Source/AnimNext/Internal/Component/AnimNextWorldSubsystem.h
--- a/Engine/Plugins/Experimental/AnimNext/Source/AnimNext/Internal/Component/AnimNextWorldSubsystem.h
+++ b/Engine/Plugins/Experimental/AnimNext/Source/AnimNext/Internal/Component/AnimNextWorldSubsystem.h
@@ -17,7 +17,7 @@ struct FRigUnit_AnimNextInitializeEvent;
 
 namespace UE::AnimNext
 {
-	struct FModuleEndTickFunction;
+	struct FModuleEventTickFunction;
 	struct FInjectionRequest;
 }
 
@@ -147,9 +147,6 @@ protected:
 	// @param	InSubsequentHandle		The handle of the module that would have been executed second
 	void RemoveModuleDependencyHandle(UE::AnimNext::FModuleHandle InPrerequisiteHandle, UE::AnimNext::FModuleHandle InSubsequentHandle);
 
-	// Disables the initialize event for a module once it has run once
-	void DisableInitializeEvent(UE::AnimNext::FModuleHandle InHandle);
-
 	// Find the const tick function for the specified event. Useful to analyze/log prerequisites. Use AddDependencyHandle, RemoveDependencyHandle to actually modify dependencies
 	// @param	InHandle			The handle of the module
 	// @param	InEventName			The event associated to the wanted tick function
@@ -189,7 +186,7 @@ protected:
 
 protected:
 	friend struct FRigUnit_AnimNextInitializeEvent;
-	friend struct UE::AnimNext::FModuleEndTickFunction;
+	friend struct UE::AnimNext::FModuleEventTickFunction;
 	friend struct UE::AnimNext::FInjectionRequest;
 
 	// Currently running instances, pooled
diff --git a/Engine/Plugins/Experimental/AnimNext/Source/AnimNext/Internal/Module/AnimNextModule.h b/Engine/Plugins/Experimental/AnimNext/Source/AnimNext/Internal/Module/AnimNextModule.h
--- a/Engine/Plugins/Experimental/AnimNext/Source/AnimNext/Internal/Module/AnimNextModule.h
+++ b/Engine/Plugins/Experimental/AnimNext/Source/AnimNext/Internal/Module/AnimNextModule.h
@@ -7,6 +7,7 @@
 #include "RigVMCore/RigVM.h"
 #include "RigVMHost.h"
 #include "Graph/AnimNextGraphState.h"
+#include "Module/RigVMTrait_ModuleEventDependency.h"
 
 #include "AnimNextModule.generated.h"
 
@@ -77,7 +78,11 @@ protected:
 	// All components that are required on startup for this module
 	UPROPERTY()
 	TArray<TObjectPtr<UScriptStruct>> RequiredComponents;
-	
+
+	// All dependencies that should be set up when the module initializes
+	UPROPERTY()
+	TArray<TInstancedStruct<FRigVMTrait_ModuleEventDependency>> Dependencies;
+
 #if WITH_EDITORONLY_DATA
 	UPROPERTY()
 	FAnimNextGraphState DefaultState_DEPRECATED;
diff --git a/Engine/Plugins/Experimental/AnimNext/Source/AnimNext/Internal/Module/AnimNextModuleContextData.h b/Engine/Plugins/Experimental/AnimNext/Source/AnimNext/Internal/Module/AnimNextModuleContextData.h
--- a/Engine/Plugins/Experimental/AnimNext/Source/AnimNext/Internal/Module/AnimNextModuleContextData.h
+++ b/Engine/Plugins/Experimental/AnimNext/Source/AnimNext/Internal/Module/AnimNextModuleContextData.h
@@ -3,7 +3,6 @@
 #pragma once
 
 #include "CoreMinimal.h"
-#include "AnimNextModuleInstance.h"
 #include "DataInterface/AnimNextDataInterfaceInstance.h"
 #include "AnimNextModuleContextData.generated.h"
 
@@ -16,23 +15,12 @@ struct FAnimNextModuleContextData
 
 	FAnimNextModuleContextData() = default;
 
-	explicit FAnimNextModuleContextData(FAnimNextModuleInstance* InModuleInstance)
-		: ModuleInstance(InModuleInstance)
-		, DataInterfaceInstance(InModuleInstance)
-	{
-	}
-	
-	FAnimNextModuleContextData(FAnimNextModuleInstance* InModuleInstance, const FAnimNextDataInterfaceInstance* InDataInterfaceInstance)
-		: ModuleInstance(InModuleInstance)
-		, DataInterfaceInstance(InDataInterfaceInstance)
-	{
-	}
+	ANIMNEXT_API explicit FAnimNextModuleContextData(FAnimNextModuleInstance* InModuleInstance);
+
+	ANIMNEXT_API FAnimNextModuleContextData(FAnimNextModuleInstance* InModuleInstance, const FAnimNextDataInterfaceInstance* InDataInterfaceInstance);
 
 	// Get the object that the module instance is bound to, if any
-	UObject* GetObject() const
-	{
-		return ModuleInstance ? ModuleInstance->GetObject() : nullptr;
-	}
+	ANIMNEXT_API UObject* GetObject() const;
 
 	FAnimNextModuleInstance& GetModuleInstance() const
 	{
diff --git a/Engine/Plugins/Experimental/AnimNext/Source/AnimNext/Internal/Module/AnimNextModuleInstance.h b/Engine/Plugins/Experimental/AnimNext/Source/AnimNext/Internal/Module/AnimNextModuleInstance.h
--- a/Engine/Plugins/Experimental/AnimNext/Source/AnimNext/Internal/Module/AnimNextModuleInstance.h
+++ b/Engine/Plugins/Experimental/AnimNext/Source/AnimNext/Internal/Module/AnimNextModuleInstance.h
@@ -15,7 +15,6 @@
 #include "StructUtils/InstancedStruct.h"
 #include "Module/AnimNextModuleInitMethod.h"
 #include "Module/TaskRunLocation.h"
-
 #include "AnimNextModuleInstance.generated.h"
 
 struct FAnimNextModuleInstance;
@@ -42,30 +41,6 @@ namespace UE::AnimNext::Debug
 	struct FDebugDraw;
 }
 
-namespace UE::AnimNext
-{
-
-// Phase is used as a general ordering constraint on event execution
-enum class EModuleEventPhase : uint8
-{
-	// Before any execution, e.g. for copying data from the game thread
-	PreExecute,
-
-	// Any initial setup work for new instances
-	Initialize,
-
-	// General execution, e.g. a prephysics event
-	Execute,
-};
-
-}
-
-namespace UE::AnimNext::Private
-{
-	// Pre-caches all possible module event info for events that have been predefined in code  
-	extern void CacheAllModuleEvents();
-}
-
 namespace UE::AnimNext
 {
 	using ModuleInstanceComponentMapType = TMap<FName, TInstancedStruct<FAnimNextModuleInstanceComponent>>;
@@ -133,12 +108,12 @@ struct FAnimNextModuleInstance : public FAnimNextDataInterfaceInstance
 	// Run a simple task on the GT via FFunctionGraphTask::CreateAndDispatchWhenReady
 	ANIMNEXT_API static void RunTaskOnGameThread(TUniqueFunction<void(void)>&& InFunction);
 
-	// Get the names of all available module names
-	static void GetModuleEventNames(TArray<FString>& OutNames, bool bIncludeInternal = false);
-
 	// Find the tick function for the specified event
-	UE::AnimNext::FModuleEventTickFunction* FindTickFunction(FName InEventName);
-	const UE::AnimNext::FModuleEventTickFunction* FindTickFunction(FName InEventName) const;
+	UE::AnimNext::FModuleEventTickFunction* FindTickFunctionByName(FName InEventName);
+	const UE::AnimNext::FModuleEventTickFunction* FindTickFunctionByName(FName InEventName) const;
+
+	// Find the first 'user' tick function
+	ANIMNEXT_API UE::AnimNext::FModuleEventTickFunction* FindFirstUserTickFunction();
 
 	// Run the specified RigVM event
 	ANIMNEXT_API void RunRigVMEvent(FName InEventName, float InDeltaTime);
@@ -152,17 +127,17 @@ struct FAnimNextModuleInstance : public FAnimNextDataInterfaceInstance
 	ANIMNEXT_API UE::AnimNext::FModuleHandle GetHandle() const { return Handle; }
 
 	// Queue a task to run at a particular module event
-	// @param	InModuleEventName	The name of the event in the module to run the supplied task relative to. If this is NAME_None, then the first valid event will be used.
+	// @param	InModuleEvent		The event in the module to run the supplied task relative to. If this is NAME_None or not found then the first 'user' event will be used.
 	// @param	InTaskFunction		The function to run
 	// @param	InLocation			Where to run the task, before or after
-	ANIMNEXT_API void QueueTask(FName InModuleEventName, TUniqueFunction<void(const UE::AnimNext::FModuleTaskContext&)>&& InTaskFunction, UE::AnimNext::ETaskRunLocation InLocation = UE::AnimNext::ETaskRunLocation::Before);
+	ANIMNEXT_API void QueueTask(FName InEventName, TUniqueFunction<void(const UE::AnimNext::FModuleTaskContext&)>&& InTaskFunction, UE::AnimNext::ETaskRunLocation InLocation = UE::AnimNext::ETaskRunLocation::Before);
 
 	// Queue a task to run at a particular module event on some other module
 	// @param	InModuleHandle		The handle for the module to queue the task on
-	// @param	InModuleEventName	The name of the event in the module to run the supplied task relative to. If this is NAME_None, then the first valid event will be used.
+	// @param	InModuleEventName	The event in the module to run the supplied task relative to. If this is NAME_None or not found then the first user 'event' will be used.
 	// @param	InTaskFunction		The function to run
 	// @param	InLocation			Where to run the task, before or after
-	ANIMNEXT_API void QueueTaskOnOtherModule(const UE::AnimNext::FModuleHandle InOtherModuleHandle, FName InModuleEventName, TUniqueFunction<void(const UE::AnimNext::FModuleTaskContext&)>&& InTaskFunction, UE::AnimNext::ETaskRunLocation InLocation = UE::AnimNext::ETaskRunLocation::Before);
+	ANIMNEXT_API void QueueTaskOnOtherModule(const UE::AnimNext::FModuleHandle InOtherModuleHandle, FName InEventName, TUniqueFunction<void(const UE::AnimNext::FModuleTaskContext&)>&& InTaskFunction, UE::AnimNext::ETaskRunLocation InLocation = UE::AnimNext::ETaskRunLocation::Before);
 
 private:
 	// Setup the entry
@@ -190,9 +165,6 @@ private:
 	void OnModuleCompiled();
 #endif
 
-	// Called to disable the initialize event for this module on the game thread when it has executed for the first time
-	void DisableInitializeEvent();
-
 	// Give the module components a chance to handle events and call their OnEndExecution extension point
 	void EndExecution(float InDeltaTime);
 
@@ -240,9 +212,6 @@ private:
 	// Copy of the handle that represents this entry to client systems
 	UE::AnimNext::FModuleHandle Handle;
 
-	// End tick function used to perform internal bookkeeping
-	UE::AnimNext::FModuleEndTickFunction EndTickFunction;
-
 	// Pre-allocated graph of tick functions
 	TArray<UE::AnimNext::FModuleEventTickFunction> TickFunctions;
 
@@ -283,18 +252,29 @@ private:
 	UPROPERTY()
 	TMap<FName, TInstancedStruct<FAnimNextModuleInstanceComponent>> ComponentMap;
 
-	// Initialize event index (INDEX_NONE if no initialize event is found)
-	int32 InitializeEventIndex = INDEX_NONE;
-
-	enum class ERunState : uint8
+	enum class EInitState : uint8
 	{
-		None,
+		NotInitialized,
 
 		CreatingTasks,
 
 		BindingTasks,
 
-		PendingInitialUpdate,
+		SetupVariables,
+
+		PendingInitializeEvent,
+
+		FirstUpdate,
+
+		Initialized
+	};
+
+	// Current initialization state
+	EInitState InitState = EInitState::NotInitialized;
+
+	enum class ERunState : uint8
+	{
+		NotInitialized,
 
 		Running,
 
@@ -302,7 +282,10 @@ private:
 	};
 
 	// Current running state
-	ERunState RunState = ERunState::None;
+	ERunState RunState = ERunState::NotInitialized;
+
+	// Transition to the specified init state, verifying that the current state is valid
+	void TransitionToInitState(EInitState InNewState);
 
 	// Transition to the specified run state, verifying that the current state is valid
 	void TransitionToRunState(ERunState InNewState);
diff --git a/Engine/Plugins/Experimental/AnimNext/Source/AnimNext/Internal/Module/ModuleEventTickFunctionBindings.h b/Engine/Plugins/Experimental/AnimNext/Source/AnimNext/Internal/Module/ModuleEventTickFunctionBindings.h
--- a/Engine/Plugins/Experimental/AnimNext/Source/AnimNext/Internal/Module/ModuleEventTickFunctionBindings.h
+++ b/Engine/Plugins/Experimental/AnimNext/Source/AnimNext/Internal/Module/ModuleEventTickFunctionBindings.h
@@ -14,11 +14,10 @@ namespace UE::AnimNext
 // Binding context for tick functions
 struct FTickFunctionBindingContext
 {
-	FTickFunctionBindingContext(FAnimNextModuleInstance& InModuleInstance, UObject* InObject, UWorld* InWorld, int32 InEventIndex)
+	FTickFunctionBindingContext(FAnimNextModuleInstance& InModuleInstance, UObject* InObject, UWorld* InWorld)
 		: ModuleInstance(InModuleInstance)
 		, Object(InObject)
 		, World(InWorld)
-		, EventIndex(InEventIndex)
 	{}
 
 	// The module instance
@@ -29,9 +28,6 @@ struct FTickFunctionBindingContext
 
 	// World that the object is contained within
 	UWorld* World = nullptr;
-
-	// Index into the module's events array
-	int32 EventIndex = INDEX_NONE;
 };
 
 // Function called to bind a tick function to a module event, set up prerequisites etc.
diff --git a/Engine/Plugins/Experimental/AnimNext/Source/AnimNext/Internal/Module/ModuleTickFunction.h b/Engine/Plugins/Experimental/AnimNext/Source/AnimNext/Internal/Module/ModuleTickFunction.h
--- a/Engine/Plugins/Experimental/AnimNext/Source/AnimNext/Internal/Module/ModuleTickFunction.h
+++ b/Engine/Plugins/Experimental/AnimNext/Source/AnimNext/Internal/Module/ModuleTickFunction.h
@@ -35,7 +35,9 @@ struct FModuleEventTickFunction : public FTickFunction
 	virtual void ExecuteTick(float DeltaTime, ELevelTick TickType, ENamedThreads::Type CurrentThread, const FGraphEventRef& MyCompletionGraphEvent) override;
 	virtual FString DiagnosticMessage() override;
 
+	void Initialize(float DeltaTime);
 	void Run(float DeltaTime);
+	void EndTick(float DeltaTime);
 
 #if WITH_EDITOR
 	// Called standalone to run a whole module's init & work in one call. Expensive and editor only.
@@ -58,30 +60,12 @@ struct FModuleEventTickFunction : public FTickFunction
 
 	FName EventName;
 
-	bool bUserEvent = false;
+	uint8 bLastUserEvent : 1 = false;
+	uint8 bFirstUserEvent : 1 = false;
+	uint8 bUserEvent : 1 = false;
 
 	// External dependencies that this tick function needs to unregister when it is destroyed
 	TArray<FTickPrerequisite> ExternalSubsequents;
 };
 
-struct FModuleEndTickFunction : public FTickFunction
-{
-	FModuleEndTickFunction()
-	{
-		ModuleInstance = nullptr;
-		bCanEverTick = true;
-		bStartWithTickEnabled = true;
-		bRunOnAnyThread = true;
-	}
-
-	// FTickFunction interface
-	virtual void ExecuteTick(float DeltaTime, ELevelTick TickType, ENamedThreads::Type CurrentThread, const FGraphEventRef& MyCompletionGraphEvent) override;
-	virtual FString DiagnosticMessage() override;
-
-	void Run(float DeltaTime);
-
-	FAnimNextModuleInstance* ModuleInstance = nullptr;
-};
-
-
 }
\ No newline at end of file
diff --git a/Engine/Plugins/Experimental/AnimNext/Source/AnimNext/Internal/Module/RigUnit_AnimNextModuleEvents.h b/Engine/Plugins/Experimental/AnimNext/Source/AnimNext/Internal/Module/RigUnit_AnimNextModuleEvents.h
--- a/Engine/Plugins/Experimental/AnimNext/Source/AnimNext/Internal/Module/RigUnit_AnimNextModuleEvents.h
+++ b/Engine/Plugins/Experimental/AnimNext/Source/AnimNext/Internal/Module/RigUnit_AnimNextModuleEvents.h
@@ -1,13 +1,27 @@
-﻿// Copyright Epic Games, Inc. All Rights Reserved.
+// Copyright Epic Games, Inc. All Rights Reserved.
 
 #pragma once
 
 #include "Graph/RigUnit_AnimNextBase.h"
 #include "AnimNextExecuteContext.h"
-#include "AnimNextModuleInstance.h"
 #include "ModuleEventTickFunctionBindings.h"
 #include "RigUnit_AnimNextModuleEvents.generated.h"
 
+namespace UE::AnimNext
+{
+
+// Phase is used as a general ordering constraint on event execution
+enum class EModuleEventPhase : uint8
+{
+	// Before any execution, e.g. for copying data from the game thread
+	PreExecute,
+
+	// General execution, e.g. a prephysics event
+	Execute,
+};
+
+}
+
 /** Base schedule-level event, never instantiated */
 USTRUCT(meta=(Abstract, Category="Events", NodeColor="1, 0, 0", Keywords="Begin,Update,Tick,Forward,Event"))
 struct ANIMNEXT_API FRigUnit_AnimNextModuleEventBase : public FRigUnit_AnimNextBase
@@ -18,7 +32,7 @@ struct ANIMNEXT_API FRigUnit_AnimNextModuleEventBase : public FRigUnit_AnimNextB
 
 	// FRigUnit interface
 	virtual FString GetUnitLabel() const override { return GetEventName().ToString(); };
-	virtual bool CanOnlyExistOnce() const override final { return true; }
+	virtual bool CanOnlyExistOnce() const override { return true; }
 
 	// Get the general ordering phase of this event, used for linearization
 	virtual UE::AnimNext::EModuleEventPhase GetEventPhase() const { return UE::AnimNext::EModuleEventPhase::Execute; }
@@ -26,6 +40,21 @@ struct ANIMNEXT_API FRigUnit_AnimNextModuleEventBase : public FRigUnit_AnimNextB
 	// Overriden in derived classes to provide binding function
 	virtual UE::AnimNext::FModuleEventBindingFunction GetBindingFunction() const { return [](const UE::AnimNext::FTickFunctionBindingContext& InContext, FTickFunction& InTickFunction){}; }
 
+	// Get the tick group for this event
+	virtual ETickingGroup GetTickGroup() const { return TG_PrePhysics; }
+
+	// Get the integer sort key for this event (used to sort events in tick groups - smaller values get sorted earlier)
+	virtual int32 GetSortOrder() const { return 0; }
+
+	// Get whether this is a user-generated event (as opposed to a compiler-generated event)
+	virtual bool IsUserEvent() const { return true; }
+
+	// Get whether this event corresponds to a separate task (so tick functions will be generated to run it)
+	virtual bool IsTask() const { return true; }
+
+	// Get whether this event is limited to the game thread only
+	virtual bool IsGameThreadTask() const { return false; }
+
 	// The execution result
 	UPROPERTY(EditAnywhere, DisplayName = "Execute", Category = "Events", meta = (Output))
 	FAnimNextExecuteContext ExecuteContext;
@@ -44,6 +73,9 @@ struct ANIMNEXT_API FRigUnit_AnimNextExecuteBindings_GT : public FRigUnit_AnimNe
 	RIGVM_METHOD()
 	virtual void Execute() override;
 	virtual FName GetEventName() const override final { return EventName; }
+	virtual int32 GetSortOrder() const override final { return 0; }	// Sort before WT bindings
+	virtual bool IsUserEvent() const override final { return false; }
+	virtual bool IsGameThreadTask() const override final { return true; }
 
 	// FRigUnit_AnimNextModuleEventBase interface
 	virtual UE::AnimNext::EModuleEventPhase GetEventPhase() const override final { return UE::AnimNext::EModuleEventPhase::PreExecute; }
@@ -63,6 +95,8 @@ struct ANIMNEXT_API FRigUnit_AnimNextExecuteBindings_WT : public FRigUnit_AnimNe
 	RIGVM_METHOD()
 	virtual void Execute() override;
 	virtual FName GetEventName() const override final { return EventName; }
+	virtual int32 GetSortOrder() const override final { return 1; }	// Sort after GT bindings
+	virtual bool IsUserEvent() const override final { return false; }
 
 	// FRigUnit_AnimNextModuleEventBase interface
 	virtual UE::AnimNext::EModuleEventPhase GetEventPhase() const override final { return UE::AnimNext::EModuleEventPhase::PreExecute; }
@@ -80,43 +114,72 @@ struct ANIMNEXT_API FRigUnit_AnimNextInitializeEvent : public FRigUnit_AnimNextM
 	RIGVM_METHOD()
 	virtual void Execute() override;
 	virtual FName GetEventName() const override final { return EventName; }
+	virtual bool IsTask() const override final { return false; }
+};
+
+/** Base event for all user-authored events. Can execute in a particular tick group (e.g. TG_PrePhysics) */
+USTRUCT(meta=(Abstract))
+struct ANIMNEXT_API FRigUnit_AnimNextUserEvent : public FRigUnit_AnimNextModuleEventBase
+{
+	GENERATED_BODY()
 
 	// FRigUnit_AnimNextModuleEventBase interface
-	virtual UE::AnimNext::EModuleEventPhase GetEventPhase() const override final { return UE::AnimNext::EModuleEventPhase::Initialize; }
+	virtual FName GetEventName() const override final { return Name; }
+	virtual FString GetUnitLabel() const override { return GetEventName().ToString(); }
+	virtual FString GetUnitSubTitle() const override;
+	virtual bool CanOnlyExistOnce() const override { return false; }
+	virtual UE::AnimNext::EModuleEventPhase GetEventPhase() const override final { return UE::AnimNext::EModuleEventPhase::Execute; }
 	virtual UE::AnimNext::FModuleEventBindingFunction GetBindingFunction() const override final;
+	virtual ETickingGroup GetTickGroup() const override final { return TickGroup; }
+	virtual int32 GetSortOrder() const override final { return SortOrder; }
+
+	// The name of the event
+	UPROPERTY(EditAnywhere, Category = "Event", meta = (Input, Constant, DetailsOnly))
+	FName Name = NAME_None;
+
+	// Sort index for ordering with other events in this tick group - smaller values get sorted earlier
+	UPROPERTY(EditAnywhere, Category = "Event", meta = (Input, Constant, DetailsOnly))
+	int32 SortOrder = 0;
+
+	// The tick group the event executes in 
+	UPROPERTY()
+	TEnumAsByte<ETickingGroup> TickGroup;
 };
 
 /** Schedule event called before world physics is updated */
 USTRUCT(meta=(DisplayName="PrePhysics", Keywords="Start,Before"))
-struct ANIMNEXT_API FRigUnit_AnimNextPrePhysicsEvent : public FRigUnit_AnimNextModuleEventBase
+struct ANIMNEXT_API FRigUnit_AnimNextPrePhysicsEvent : public FRigUnit_AnimNextUserEvent
 {
 	GENERATED_BODY()
 
-	static inline const FLazyName EventName = FLazyName("PrePhysics");
+	static inline const FLazyName DefaultEventName = FLazyName("PrePhysics");
+
+	FRigUnit_AnimNextPrePhysicsEvent()
+	{
+		Name = DefaultEventName;
+		TickGroup = ETickingGroup::TG_PrePhysics;
+	}
 
 	RIGVM_METHOD()
 	virtual void Execute() override;
-	virtual FName GetEventName() const override final { return EventName; }
-
-	// FRigUnit_AnimNextModuleEventBase interface
-	virtual UE::AnimNext::EModuleEventPhase GetEventPhase() const override final { return UE::AnimNext::EModuleEventPhase::Execute; }
-	virtual UE::AnimNext::FModuleEventBindingFunction GetBindingFunction() const override final;
 };
 
 /** Schedule event called after world physics is updated */
 USTRUCT(meta=(DisplayName="PostPhysics", Keywords="End,After"))
-struct ANIMNEXT_API FRigUnit_AnimNextPostPhysicsEvent : public FRigUnit_AnimNextModuleEventBase
+struct ANIMNEXT_API FRigUnit_AnimNextPostPhysicsEvent : public FRigUnit_AnimNextUserEvent
 {
 	GENERATED_BODY()
 
-	static inline const FLazyName EventName = FLazyName("PostPhysics");
+	static inline const FLazyName DefaultEventName = FLazyName("PostPhysics");
 
+	FRigUnit_AnimNextPostPhysicsEvent()
+	{
+		Name = DefaultEventName;
+		TickGroup = ETickingGroup::TG_PostPhysics;
+	}
+	
 	RIGVM_METHOD()
 	virtual void Execute() override;
-	virtual FName GetEventName() const override final { return EventName; }
 
-	// FRigUnit_AnimNextModuleEventBase interface
-	virtual UE::AnimNext::EModuleEventPhase GetEventPhase() const override final { return UE::AnimNext::EModuleEventPhase::Execute; }
-	virtual UE::AnimNext::FModuleEventBindingFunction GetBindingFunction() const override final;
 };
 
diff --git a/Engine/Plugins/Experimental/AnimNext/Source/AnimNext/Private/AnimNextConfig.cpp b/Engine/Plugins/Experimental/AnimNext/Source/AnimNext/Private/AnimNextConfig.cpp
--- a/Engine/Plugins/Experimental/AnimNext/Source/AnimNext/Private/AnimNextConfig.cpp
+++ b/Engine/Plugins/Experimental/AnimNext/Source/AnimNext/Private/AnimNextConfig.cpp
@@ -12,10 +12,3 @@ void UAnimNextConfig::PostEditChangeProperty(FPropertyChangedEvent& PropertyChan
 }
 
 #endif
-
-TArray<FString> UAnimNextConfig::GetModuleEventNames()
-{
-	TArray<FString> EventNames;
-	FAnimNextModuleInstance::GetModuleEventNames(EventNames);
-	return EventNames;
-}
\ No newline at end of file
diff --git a/Engine/Plugins/Experimental/AnimNext/Source/AnimNext/Private/AnimNextModuleImpl.cpp b/Engine/Plugins/Experimental/AnimNext/Source/AnimNext/Private/AnimNextModuleImpl.cpp
--- a/Engine/Plugins/Experimental/AnimNext/Source/AnimNext/Private/AnimNextModuleImpl.cpp
+++ b/Engine/Plugins/Experimental/AnimNext/Source/AnimNext/Private/AnimNextModuleImpl.cpp
@@ -71,7 +71,6 @@ namespace UE::AnimNext
 
 		RigVMRegistry.RegisterStructTypes(AllowedStructTypes);
 
-		Private::CacheAllModuleEvents();
 		FDataRegistry::Init();
 
 		FRigVMRuntimeDataRegistry::Init();
diff --git a/Engine/Plugins/Experimental/AnimNext/Source/AnimNext/Private/Component/AnimNextWorldSubsystem.cpp b/Engine/Plugins/Experimental/AnimNext/Source/AnimNext/Private/Component/AnimNextWorldSubsystem.cpp
--- a/Engine/Plugins/Experimental/AnimNext/Source/AnimNext/Private/Component/AnimNextWorldSubsystem.cpp
+++ b/Engine/Plugins/Experimental/AnimNext/Source/AnimNext/Private/Component/AnimNextWorldSubsystem.cpp
@@ -272,46 +272,7 @@ void UAnimNextWorldSubsystem::QueueTaskHandle(UE::AnimNext::FModuleHandle InOutH
 	if (IsValidHandle(InOutHandle))
 	{
 		FAnimNextModuleInstance& Instance = Instances.Get(InOutHandle);
-
-		TSpscQueue<TUniqueFunction<void(const FModuleTaskContext&)>>* Queue = nullptr;
-		FModuleEventTickFunction* FoundTickFunction = nullptr;
-		if(Instance.TickFunctions.Num() > 0)
-		{
-			if(InModuleEventName == NAME_None)
-			{
-				// passing NAME_None means 'any tick function' in this context, so just use the first one
-				FoundTickFunction = &Instance.TickFunctions[0];
-			}
-			else
-			{
-				FoundTickFunction = Instance.TickFunctions.FindByPredicate([InModuleEventName](const FModuleEventTickFunction& InTickFunction)
-				{
-					return InTickFunction.EventName == InModuleEventName;
-				});
-			}
-		}
-
-		if (FoundTickFunction)
-		{
-			switch (InLocation)
-			{
-			case ETaskRunLocation::Before:
-				Queue = &FoundTickFunction->PreExecuteTasks;
-				break;
-			case ETaskRunLocation::After:
-				Queue = &FoundTickFunction->PostExecuteTasks;
-				break;
-			}
-		}
-
-		if (Queue)
-		{
-			Queue->Enqueue(MoveTemp(InTaskFunction));
-		}
-		else
-		{
-			UE_LOGFMT(LogAnimation, Warning, "QueueTask: Could not find event '{EventName}' in module '{ModuleName}'", InModuleEventName, Instance.GetDataInterfaceName());
-		}
+		Instance.QueueTask(InModuleEventName, MoveTemp(InTaskFunction), InLocation);
 	}
 }
 
@@ -333,17 +294,6 @@ void UAnimNextWorldSubsystem::OnModuleCompiled(UAnimNextModule* InModule)
 
 #endif
 
-void UAnimNextWorldSubsystem::DisableInitializeEvent(UE::AnimNext::FModuleHandle InHandle)
-{
-	using namespace UE::AnimNext;
-	check(IsInGameThread());
-	if (IsValidHandle(InHandle))
-	{
-		FAnimNextModuleInstance& Instance = Instances.Get(InHandle);
-		Instance.DisableInitializeEvent();
-	}
-}
-
 void UAnimNextWorldSubsystem::QueueInputTraitEventHandle(UE::AnimNext::FModuleHandle InHandle, FAnimNextTraitEventPtr Event)
 {
 	using namespace UE::AnimNext;
@@ -386,7 +336,7 @@ const FTickFunction* UAnimNextWorldSubsystem::FindTickFunctionHandle(UE::AnimNex
 	if (IsValidHandle(InHandle))
 	{
 		const FAnimNextModuleInstance& Instance = Instances.Get(InHandle);
-		const FModuleEventTickFunction* TickFunction = Instance.FindTickFunction(InEventName);
+		const FModuleEventTickFunction* TickFunction = Instance.FindTickFunctionByName(InEventName);
 		if(TickFunction == nullptr)
 		{
 			UE_LOGFMT(LogAnimation, Warning, "FindTickFunctionHandle: Could not find event '{EventName}' in module '{ModuleName}'", InEventName, Instance.GetDataInterfaceName());
@@ -410,7 +360,7 @@ void UAnimNextWorldSubsystem::AddDependencyHandle(UE::AnimNext::FModuleHandle In
 	if (IsValidHandle(InHandle))
 	{
 		FAnimNextModuleInstance& Instance = Instances.Get(InHandle);
-		FModuleEventTickFunction* TickFunction = Instance.FindTickFunction(InEventName);
+		FModuleEventTickFunction* TickFunction = Instance.FindTickFunctionByName(InEventName);
 		if(TickFunction == nullptr)
 		{
 			UE_LOGFMT(LogAnimation, Warning, "AddDependencyHandle: Could not find event '{EventName}' in module '{ModuleName}'", InEventName, Instance.GetDataInterfaceName());
@@ -435,7 +385,7 @@ void UAnimNextWorldSubsystem::RemoveDependencyHandle(UE::AnimNext::FModuleHandle
 	if (IsValidHandle(InHandle))
 	{
 		FAnimNextModuleInstance& Instance = Instances.Get(InHandle);
-		FModuleEventTickFunction* TickFunction = Instance.FindTickFunction(InEventName);
+		FModuleEventTickFunction* TickFunction = Instance.FindTickFunctionByName(InEventName);
 		if(TickFunction == nullptr)
 		{
 			UE_LOGFMT(LogAnimation, Warning, "RemoveDependencyHandle: Could not find event '{EventName}' in module '{ModuleName}'", InEventName, Instance.GetDataInterfaceName());
@@ -461,7 +411,7 @@ void UAnimNextWorldSubsystem::AddModuleEventDependencyHandle(UE::AnimNext::FModu
 	if (IsValidHandle(OtherHandle))
 	{
 		FAnimNextModuleInstance& Instance = Instances.Get(OtherHandle);
-		FModuleEventTickFunction* TickFunction = Instance.FindTickFunction(OtherEventName);
+		FModuleEventTickFunction* TickFunction = Instance.FindTickFunctionByName(OtherEventName);
 		if(TickFunction == nullptr)
 		{
 			UE_LOGFMT(LogAnimation, Warning, "AddDependencyHandle: Could not find event '{EventName}' in module '{ModuleName}'", OtherEventName, Instance.GetDataInterfaceName());
@@ -480,7 +430,7 @@ void UAnimNextWorldSubsystem::RemoveModuleEventDependencyHandle(UE::AnimNext::FM
 	if (IsValidHandle(InHandle))
 	{
 		FAnimNextModuleInstance& Instance = Instances.Get(InHandle);
-		FModuleEventTickFunction* TickFunction = Instance.FindTickFunction(InEventName);
+		FModuleEventTickFunction* TickFunction = Instance.FindTickFunctionByName(InEventName);
 		if(TickFunction == nullptr)
 		{
 			UE_LOGFMT(LogAnimation, Warning, "RemoveDependencyHandle: Could not find event '{EventName}' in module '{ModuleName}'", OtherEventName, Instance.GetDataInterfaceName());
diff --git a/Engine/Plugins/Experimental/AnimNext/Source/AnimNext/Private/Module/AnimNextModuleContextData.cpp b/Engine/Plugins/Experimental/AnimNext/Source/AnimNext/Private/Module/AnimNextModuleContextData.cpp
new file mode 100644
--- /dev/null
+++ b/Engine/Plugins/Experimental/AnimNext/Source/AnimNext/Private/Module/AnimNextModuleContextData.cpp
@@ -0,0 +1,21 @@
+﻿// Copyright Epic Games, Inc. All Rights Reserved.
+
+#include "Module/AnimNextModuleContextData.h"
+#include "Module/AnimNextModuleInstance.h"
+
+FAnimNextModuleContextData::FAnimNextModuleContextData(FAnimNextModuleInstance* InModuleInstance)
+	: ModuleInstance(InModuleInstance)
+	, DataInterfaceInstance(InModuleInstance)
+{
+}
+
+FAnimNextModuleContextData::FAnimNextModuleContextData(FAnimNextModuleInstance* InModuleInstance, const FAnimNextDataInterfaceInstance* InDataInterfaceInstance)
+	: ModuleInstance(InModuleInstance)
+	, DataInterfaceInstance(InDataInterfaceInstance)
+{
+}
+
+UObject* FAnimNextModuleContextData::GetObject() const
+{
+	return ModuleInstance ? ModuleInstance->GetObject() : nullptr;
+}
\ No newline at end of file
diff --git a/Engine/Plugins/Experimental/AnimNext/Source/AnimNext/Private/Module/AnimNextModuleInstance.cpp b/Engine/Plugins/Experimental/AnimNext/Source/AnimNext/Private/Module/AnimNextModuleInstance.cpp
--- a/Engine/Plugins/Experimental/AnimNext/Source/AnimNext/Private/Module/AnimNextModuleInstance.cpp
+++ b/Engine/Plugins/Experimental/AnimNext/Source/AnimNext/Private/Module/AnimNextModuleInstance.cpp
@@ -33,7 +33,8 @@ FAnimNextModuleInstance::FAnimNextModuleInstance(
 	: Object(InObject)
 	, Pool(InPool)
 	, ProxyHost(InProxyHost)
-	, RunState(ERunState::None)
+	, InitState(EInitState::NotInitialized)
+	, RunState(ERunState::NotInitialized)
 	, InitMethod(InInitMethod)
 {
 	check(InModule);
@@ -71,57 +72,69 @@ struct FImplementedModuleEvent
 	FModuleEventBindingFunction Binding;
 	FName EventName;
 	EModuleEventPhase Phase = EModuleEventPhase::Execute;
+	ETickingGroup TickGroup = ETickingGroup::TG_PrePhysics;
+	int32 SortOrder = 0;
 	bool bUserEvent = false;
+	bool bIsTask = false;
+	bool bIsGameThreadTask = false;
 };
 
-static TArray<FImplementedModuleEvent> GAllModuleEvents;
 static TArray<FImplementedModuleEvent> GImplementedModuleEvents;
 
-void CacheAllModuleEvents()
+// Gets information about the module events that are implemented by the supplied VM, sorted by execution order in the frame
+static TConstArrayView<FImplementedModuleEvent> GetImplementedModuleEvents(const URigVM* VM)
 {
 	check(IsInGameThread());	// This function cannot be run concurrently because of static usage
+	GImplementedModuleEvents.Reset();
 
-	if(GAllModuleEvents.Num() == 0)
+	const FRigVMByteCode& ByteCode = VM->GetByteCode();
+	const TArray<const FRigVMFunction*>& Functions = VM->GetFunctions();
+	FRigVMInstructionArray Instructions = ByteCode.GetInstructions();
+	for (int32 EntryIndex = 0; EntryIndex < ByteCode.NumEntries(); EntryIndex++)
 	{
-		for(TObjectIterator<UScriptStruct> It; It; ++It)
+		const FRigVMByteCodeEntry& Entry = ByteCode.GetEntry(EntryIndex);
+		const FRigVMInstruction& Instruction = Instructions[Entry.InstructionIndex];
+		const FRigVMExecuteOp& Op = ByteCode.GetOpAt<FRigVMExecuteOp>(Instruction);
+		const FRigVMFunction* Function = Functions[Op.FunctionIndex];
+		check(Function != nullptr);
+
+		if (Function->Struct->IsChildOf(FRigUnit_AnimNextModuleEventBase::StaticStruct()))
 		{
-			UScriptStruct* Struct = *It;
-			if(Struct->IsChildOf(FRigUnit_AnimNextModuleEventBase::StaticStruct()) && Struct != FRigUnit_AnimNextModuleEventBase::StaticStruct())
-			{
-				TInstancedStruct<FRigUnit_AnimNextModuleEventBase> StructInstance;
-				StructInstance.InitializeAsScriptStruct(Struct);
-				FImplementedModuleEvent& NewEvent = GAllModuleEvents.AddDefaulted_GetRef();
-				NewEvent.Struct = Struct;
-				NewEvent.Binding = StructInstance.Get().GetBindingFunction();
-				NewEvent.EventName = StructInstance.Get().GetEventName();
-				NewEvent.Phase = StructInstance.Get().GetEventPhase();
-				NewEvent.bUserEvent = Struct != FRigUnit_AnimNextExecuteBindings_GT::StaticStruct() &&
-									  Struct != FRigUnit_AnimNextExecuteBindings_WT::StaticStruct();
-			}
+			TInstancedStruct<FRigUnit_AnimNextModuleEventBase> StructInstance;
+			StructInstance.InitializeAsScriptStruct(Function->Struct);
+			const FRigUnit_AnimNextModuleEventBase& Event = StructInstance.Get();
+			FImplementedModuleEvent& NewEvent = GImplementedModuleEvents.AddDefaulted_GetRef();
+			NewEvent.Struct = Function->Struct;
+			NewEvent.Binding = Event.GetBindingFunction();
+			NewEvent.EventName = Event.GetEventName();
+			NewEvent.Phase = Event.GetEventPhase();
+			NewEvent.TickGroup = Event.GetTickGroup();
+			NewEvent.SortOrder = Event.GetSortOrder();
+			NewEvent.bUserEvent = Event.IsUserEvent();
+			NewEvent.bIsTask = Event.IsTask();
+			NewEvent.bIsGameThreadTask = Event.IsGameThreadTask();
 		}
-
-		Algo::SortBy(GAllModuleEvents, &FImplementedModuleEvent::Phase);
-		GImplementedModuleEvents.Reserve(GAllModuleEvents.Num());
 	}
-}
-
-// Gets information about the module events that are implemented by the supplied VM
-static TConstArrayView<FImplementedModuleEvent> GetImplementedModuleEvents(const URigVM* VM)
-{
-	check(IsInGameThread());	// This function cannot be run concurrently because of static usage
-	check(GAllModuleEvents.Num() > 0);	// Call CacheAllModuleEvents before this function
 
-	// Get all the module events from the VM entry points
-	const TArray<FName>& EntryNames = VM->GetEntryNames();
-	GImplementedModuleEvents.Reset();
-	for(const FImplementedModuleEvent& ModuleEvent : GAllModuleEvents)
+	Algo::Sort(GImplementedModuleEvents, [](const FImplementedModuleEvent& InA, const FImplementedModuleEvent& InB)
 	{
-		if(EntryNames.Contains(ModuleEvent.EventName))
+		if (InA.Phase != InB.Phase)
 		{
-			GImplementedModuleEvents.Add(ModuleEvent);
+			return InA.Phase < InB.Phase;
+		}
+		else if (InA.TickGroup != InB.TickGroup)
+		{
+			return InA.TickGroup < InB.TickGroup;
+		}
+		else if (InA.SortOrder != InB.SortOrder)
+		{
+			return InA.SortOrder < InB.SortOrder;
 		}
-	}
 
+		// Tie-break sorting on event name for determinism
+		return InA.EventName.Compare(InB.EventName) < 0;
+	});
+	
 	return GImplementedModuleEvents;
 }
 
@@ -148,39 +161,75 @@ void FAnimNextModuleInstance::Initialize()
 		WorldType = World->WorldType;
 	}
 
-	// Get all the module events from the VM entry points, sorted by phase
+	// Get all the module events from the VM entry points, sorted by ordering in the frame
 	URigVM* VM = Module->RigVM;
 	TConstArrayView<Private::FImplementedModuleEvent> ImplementedModuleEvents = Private::GetImplementedModuleEvents(VM);
 
 	// Setup tick function graph using module events
 	if (ImplementedModuleEvents.Num() > 0)
 	{
-		TransitionToRunState(ERunState::CreatingTasks);
+		TransitionToInitState(EInitState::CreatingTasks);
 
 		// Allocate tick functions
-		TickFunctions.SetNum(ImplementedModuleEvents.Num()); 
-		for (TEnumerateRef<const Private::FImplementedModuleEvent> ModuleEvent : EnumerateRange(ImplementedModuleEvents))
+		TickFunctions.Reserve(ImplementedModuleEvents.Num());
+		bool bFoundFirstUserEvent = false;
+		FModuleEventTickFunction* PrevTickFunction = nullptr;
+		for (int32 EventIndex = 0; EventIndex < ImplementedModuleEvents.Num(); EventIndex++)
 		{
-			FModuleEventTickFunction& TickFunction = TickFunctions[ModuleEvent.GetIndex()];
+			const Private::FImplementedModuleEvent& ModuleEvent = ImplementedModuleEvents[EventIndex];
+			if (!ModuleEvent.bIsTask)
+			{
+				continue;
+			}
+
+			FModuleEventTickFunction& TickFunction = TickFunctions.AddDefaulted_GetRef();
+			TickFunction.bRunOnAnyThread = !ModuleEvent.bIsGameThreadTask;
 			TickFunction.ModuleInstance = this;
-			TickFunction.EventName = ModuleEvent->EventName;
-			TickFunction.bUserEvent = ModuleEvent->bUserEvent;
+			TickFunction.EventName = ModuleEvent.EventName;
+			TickFunction.TickGroup = ModuleEvent.TickGroup;
+			TickFunction.bUserEvent = ModuleEvent.bUserEvent;
+
+			// Perform custom setup
+			FTickFunctionBindingContext Context(*this, Object, World);
+			ModuleEvent.Binding(Context, TickFunction);
+
+			// Establish linear dependency chain
+			if (PrevTickFunction != nullptr)
+			{
+				TickFunction.AddPrerequisite(Object, *PrevTickFunction);
+			}
+			PrevTickFunction = &TickFunction;
+
+			// Set up dependencies, if any
+			for (const TInstancedStruct<FRigVMTrait_ModuleEventDependency>& DependencyInstance : Module->Dependencies)
+			{
+				const FRigVMTrait_ModuleEventDependency* Dependency = DependencyInstance.GetPtr<FRigVMTrait_ModuleEventDependency>();
+				if (Dependency != nullptr && Dependency->EventName == ModuleEvent.EventName)
+				{
+					FModuleDependencyContext ModuleDependencyContext(Object, TickFunction);
+					Dependency->OnAddDependency(ModuleDependencyContext);
+				}
+			}
+
+			if (ModuleEvent.bUserEvent && !bFoundFirstUserEvent)
+			{
+				TickFunction.bFirstUserEvent = true;
+				bFoundFirstUserEvent = true;
+			}
 		}
 
-		// Bind tick functions (done in a second pass to allow functions to add prerequisites between each other)
-		for (TEnumerateRef<const Private::FImplementedModuleEvent> ModuleEvent : EnumerateRange(ImplementedModuleEvents))
+		// Find the last user event - 'end' logic will be called from here
+		for (int32 EventIndex = TickFunctions.Num() - 1; EventIndex >= 0; EventIndex--)
 		{
-			FModuleEventTickFunction& TickFunction = TickFunctions[ModuleEvent.GetIndex()];
-			// Add prerequisite on end tick function
-			EndTickFunction.AddPrerequisite(Object, TickFunction);
-			FTickFunctionBindingContext Context(*this, Object, World, ModuleEvent.GetIndex());
-			ModuleEvent->Binding(Context, TickFunction);
+			FModuleEventTickFunction& TickFunction = TickFunctions[EventIndex];
+			if (TickFunction.bUserEvent)
+			{
+				TickFunction.bLastUserEvent = true;
+				break;
+			}
 		}
 
-		// Bind end tick function to this instance
-		EndTickFunction.ModuleInstance = this;
-
-		TransitionToRunState(ERunState::BindingTasks);
+		TransitionToInitState(EInitState::BindingTasks);
 
 		// Register our tick functions
 		if(World)
@@ -190,10 +239,11 @@ void FAnimNextModuleInstance::Initialize()
 			{
 				TickFunction.RegisterTickFunction(Level);
 			}
-			EndTickFunction.RegisterTickFunction(Level);
 		}
 
-		TransitionToRunState(ERunState::PendingInitialUpdate);
+		TransitionToInitState(EInitState::SetupVariables);
+
+		// TODO: code in EInitState::SetupVariables phase below can probably move to FModuleEventTickFunction::Initialize
 
 		// Initialize variables
 		const int32 NumVariables = Module->VariableDefaults.GetNumPropertiesInBag();
@@ -247,6 +297,9 @@ void FAnimNextModuleInstance::Initialize()
 			AddComponentInternal(ComponentNameHash, ComponentName, MoveTemp(Component));
 		}
 
+		TransitionToInitState(EInitState::PendingInitializeEvent);
+		TransitionToRunState(ERunState::Running);
+
 		// Just pause now if we arent needing an initial update
 		if(InitMethod == EAnimNextModuleInitMethod::None)
 		{
@@ -294,7 +347,8 @@ void FAnimNextModuleInstance::ResetBindingsAndInstanceData()
 
 	FModuleWriteGuard Guard(this);
 
-	TransitionToRunState(ERunState::None);
+	TransitionToInitState(EInitState::NotInitialized);
+	TransitionToRunState(ERunState::NotInitialized);
 
 	for (FModuleEventTickFunction& TickFunction : TickFunctions)
 	{
@@ -303,8 +357,6 @@ void FAnimNextModuleInstance::ResetBindingsAndInstanceData()
 		TickFunction.UnRegisterTickFunction();
 	}
 
-	EndTickFunction.UnRegisterTickFunction();
-
 	TickFunctions.Reset();
 
 	ExtendedExecuteContext.Reset();
@@ -348,41 +400,65 @@ void FAnimNextModuleInstance::Enable(bool bInEnabled)
 
 	FModuleWriteGuard Guard(this);
 
-	if(RunState == ERunState::PendingInitialUpdate || RunState == ERunState::Paused || RunState == ERunState::Running)
+	if(RunState == ERunState::Paused || RunState == ERunState::Running)
 	{
 		for (FModuleEventTickFunction& TickFunction : TickFunctions)
 		{
 			TickFunction.SetTickFunctionEnable(bInEnabled);
 		}
-		EndTickFunction.SetTickFunctionEnable(bInEnabled);
 
 		TransitionToRunState(bInEnabled ? ERunState::Running : ERunState::Paused);
 	}
 }
 
-void FAnimNextModuleInstance::TransitionToRunState(ERunState InNewState)
+void FAnimNextModuleInstance::TransitionToInitState(EInitState InNewState)
 {
 	UE::AnimNext::FModuleWriteGuard Guard(this);
 
 	switch(InNewState)
 	{
-	case ERunState::None:
-		check(RunState == ERunState::None || RunState == ERunState::PendingInitialUpdate || RunState == ERunState::Paused || RunState == ERunState::Running);
+	case EInitState::NotInitialized:
+		check(InitState == EInitState::NotInitialized || InitState == EInitState::PendingInitializeEvent || InitState == EInitState::SetupVariables || InitState == EInitState::FirstUpdate || InitState == EInitState::Initialized);
+		break;
+	case EInitState::CreatingTasks:
+		check(InitState == EInitState::NotInitialized);
+		break;
+	case EInitState::BindingTasks:
+		check(InitState == EInitState::CreatingTasks);
+		break;
+	case EInitState::SetupVariables:
+		check(InitState == EInitState::BindingTasks);
 		break;
-	case ERunState::CreatingTasks:
-		check(RunState == ERunState::None);
+	case EInitState::PendingInitializeEvent:
+		check(InitState == EInitState::SetupVariables);
 		break;
-	case ERunState::BindingTasks:
-		check(RunState == ERunState::CreatingTasks);
+	case EInitState::FirstUpdate:
+		check(InitState == EInitState::PendingInitializeEvent);
 		break;
-	case ERunState::PendingInitialUpdate:
-		check(RunState == ERunState::BindingTasks);
+	case EInitState::Initialized:
+		check(InitState == EInitState::FirstUpdate);
 		break;
+	default:
+		checkNoEntry();
+	}
+
+	InitState = InNewState;
+}
+
+void FAnimNextModuleInstance::TransitionToRunState(ERunState InNewState)
+{
+	UE::AnimNext::FModuleWriteGuard Guard(this);
+
+	switch(InNewState)
+	{
 	case ERunState::Running:
-		check(RunState == ERunState::PendingInitialUpdate || RunState == ERunState::Paused || RunState == ERunState::Running);
+		check(RunState == ERunState::NotInitialized || RunState == ERunState::Paused || RunState == ERunState::Running);
 		break;
 	case ERunState::Paused:
-		check(RunState == ERunState::PendingInitialUpdate || RunState == ERunState::Paused || RunState == ERunState::Running);
+		check(RunState == ERunState::Paused || RunState == ERunState::Running);
+		break;
+	case ERunState::NotInitialized:
+		check(RunState == ERunState::NotInitialized || RunState == ERunState::Paused || RunState == ERunState::Running);
 		break;
 	default:
 		checkNoEntry();
@@ -515,17 +591,6 @@ void FAnimNextModuleInstance::ShowDebugDrawing(bool bInShowDebugDrawing)
 }
 #endif
 
-void FAnimNextModuleInstance::DisableInitializeEvent()
-{
-	if(InitializeEventIndex != INDEX_NONE)
-	{
-		check(IsInGameThread());
-		check(TickFunctions.IsValidIndex(InitializeEventIndex));
-		check(TickFunctions[InitializeEventIndex].EventName == FRigUnit_AnimNextInitializeEvent::EventName);
-		TickFunctions[InitializeEventIndex].SetTickFunctionEnable(false);
-	}
-}
-
 void FAnimNextModuleInstance::RunTaskOnGameThread(TUniqueFunction<void(void)>&& InFunction)
 {
 	if(IsInGameThread())
@@ -538,53 +603,20 @@ void FAnimNextModuleInstance::RunTaskOnGameThread(TUniqueFunction<void(void)>&&
 	}
 }
 
-void FAnimNextModuleInstance::GetModuleEventNames(TArray<FString>& OutNames, bool bIncludeInternal)
-{
-	using namespace UE::AnimNext::Private;
-
-	check(IsInGameThread());			// This function cannot be run concurrently because of static usage
-	check(GAllModuleEvents.Num() > 0);	// Call CacheAllModuleEvents before this function
-	
-	for(const FImplementedModuleEvent& ModuleEvent : GAllModuleEvents)
-	{
-#if WITH_METADATA
-		bool bIsHidden = ModuleEvent.Struct->HasMetaData(FRigVMStruct::HiddenMetaName);
-#else
-		bool bIsHidden = false;
-#endif
-	
-		if(bIncludeInternal || !bIsHidden)
-		{
-			OutNames.Add(ModuleEvent.EventName.ToString());
-		}
-	}
-}
-
-UE::AnimNext::FModuleEventTickFunction* FAnimNextModuleInstance::FindTickFunction(FName InEventName)
+UE::AnimNext::FModuleEventTickFunction* FAnimNextModuleInstance::FindTickFunctionByName(FName InEventName)
 {
 	using namespace UE::AnimNext;
 
-	FModuleWriteGuard Guard(this);
-
-	for(FModuleEventTickFunction& TickFunction : TickFunctions)
-	{
-		if(TickFunction.EventName == InEventName)
-		{
-			return &TickFunction;
-		}
-	}
-	return nullptr;
+	return const_cast<FModuleEventTickFunction*>(const_cast<const FAnimNextModuleInstance*>(this)->FindTickFunctionByName(InEventName));
 }
 
-const UE::AnimNext::FModuleEventTickFunction* FAnimNextModuleInstance::FindTickFunction(FName InEventName) const
+const UE::AnimNext::FModuleEventTickFunction* FAnimNextModuleInstance::FindTickFunctionByName(FName InEventName) const
 {
 	using namespace UE::AnimNext;
 
-	FModuleWriteGuard Guard(this);
-
 	for(const FModuleEventTickFunction& TickFunction : TickFunctions)
 	{
-		if(TickFunction.EventName == InEventName)
+		if (TickFunction.EventName == InEventName)
 		{
 			return &TickFunction;
 		}
@@ -632,6 +664,7 @@ void FAnimNextModuleInstance::RaiseTraitEvents(const UE::AnimNext::FTraitEventLi
 void FAnimNextModuleInstance::AddPrerequisite(FAnimNextModuleInstance& InPrerequisiteInstance)
 {
 	check(IsInGameThread());
+	check(&InPrerequisiteInstance != this);
 
 	UE::AnimNext::FModuleWriteGuard Guard(this);
 
@@ -646,11 +679,11 @@ void FAnimNextModuleInstance::AddPrerequisite(FAnimNextModuleInstance& InPrerequ
 	}
 	else
 	{
-		// Ensure all of our tick functions execute after the prerequisite's end tick function
-		EndTickFunction.AddPrerequisite(InPrerequisiteInstance.Object, InPrerequisiteInstance.EndTickFunction);
+		// Ensure all of our tick functions execute after the prerequisite's last tick function
+		TickFunctions.Last().AddPrerequisite(InPrerequisiteInstance.Object, InPrerequisiteInstance.TickFunctions.Last());
 		for(UE::AnimNext::FModuleEventTickFunction& TickFunction : TickFunctions)
 		{
-			TickFunction.AddPrerequisite(InPrerequisiteInstance.Object, InPrerequisiteInstance.EndTickFunction);
+			TickFunction.AddPrerequisite(InPrerequisiteInstance.Object, InPrerequisiteInstance.TickFunctions.Last());
 		}
 
 		PrerequisiteRefs.Add({ InPrerequisiteInstance.Handle, 1});
@@ -660,6 +693,7 @@ void FAnimNextModuleInstance::AddPrerequisite(FAnimNextModuleInstance& InPrerequ
 void FAnimNextModuleInstance::RemovePrerequisite(FAnimNextModuleInstance& InPrerequisiteInstance)
 {
 	check(IsInGameThread());
+	check(&InPrerequisiteInstance != this);
 
 	UE::AnimNext::FModuleWriteGuard Guard(this);
 
@@ -675,11 +709,11 @@ void FAnimNextModuleInstance::RemovePrerequisite(FAnimNextModuleInstance& InPrer
 		FoundReference.ReferenceCount--;
 		if(FoundReference.ReferenceCount == 0)
 		{
-			// Remove dependency on the prerequisite's end tick function
-			EndTickFunction.RemovePrerequisite(InPrerequisiteInstance.Object, InPrerequisiteInstance.EndTickFunction);
+			// Remove dependency on the prerequisite's last tick function
+			TickFunctions.Last().RemovePrerequisite(InPrerequisiteInstance.Object, InPrerequisiteInstance.TickFunctions.Last());
 			for(UE::AnimNext::FModuleEventTickFunction& TickFunction : TickFunctions)
 			{
-				TickFunction.RemovePrerequisite(InPrerequisiteInstance.Object, InPrerequisiteInstance.EndTickFunction);
+				TickFunction.RemovePrerequisite(InPrerequisiteInstance.Object, InPrerequisiteInstance.TickFunctions.Last());
 			}
 
 			PrerequisiteRefs.RemoveAtSwap(FoundIndex, EAllowShrinking::No);
@@ -760,28 +794,45 @@ TArrayView<UE::AnimNext::FModuleEventTickFunction> FAnimNextModuleInstance::GetT
 	return TickFunctions;
 }
 
-void FAnimNextModuleInstance::QueueTask(FName InModuleEventName, TUniqueFunction<void(const UE::AnimNext::FModuleTaskContext&)>&& InTaskFunction, UE::AnimNext::ETaskRunLocation InLocation)
+UE::AnimNext::FModuleEventTickFunction* FAnimNextModuleInstance::FindFirstUserTickFunction()
+{
+	UE::AnimNext::FModuleWriteGuard PrerequisiteGuard(this);
+
+	UE::AnimNext::FModuleEventTickFunction* FoundTickFunction = TickFunctions.FindByPredicate([](const UE::AnimNext::FModuleEventTickFunction& InTickFunction)
+	{
+		return InTickFunction.bFirstUserEvent;
+	});
+
+	return FoundTickFunction;
+}
+
+void FAnimNextModuleInstance::QueueTask(FName InEventName, TUniqueFunction<void(const UE::AnimNext::FModuleTaskContext&)>&& InTaskFunction, UE::AnimNext::ETaskRunLocation InLocation)
 {
 	using namespace UE::AnimNext;
 
-	TSpscQueue<TUniqueFunction<void(const FModuleTaskContext&)>>* Queue = nullptr;
 	FModuleEventTickFunction* FoundTickFunction = nullptr;
 	if(TickFunctions.Num() > 0)
 	{
-		if(InModuleEventName == NAME_None)
+		if(!InEventName.IsNone())
 		{
-			// passing NAME_None means 'any tick function' in this context, so just use the first one
-			FoundTickFunction = &TickFunctions[0];
+			// Match according to event desc
+			FoundTickFunction = TickFunctions.FindByPredicate([InEventName](const FModuleEventTickFunction& InTickFunction)
+			{
+				return InTickFunction.EventName == InEventName;
+			});
 		}
-		else
+
+		if (FoundTickFunction == nullptr)
 		{
-			FoundTickFunction = TickFunctions.FindByPredicate([InModuleEventName](const FModuleEventTickFunction& InTickFunction)
+			// Fall back to first user function
+			FoundTickFunction = TickFunctions.FindByPredicate([](const UE::AnimNext::FModuleEventTickFunction& InTickFunction)
 			{
-				return InTickFunction.EventName == InModuleEventName;
+				return InTickFunction.bFirstUserEvent;
 			});
 		}
 	}
 
+	TSpscQueue<TUniqueFunction<void(const FModuleTaskContext&)>>* Queue = nullptr;
 	if (FoundTickFunction)
 	{
 		switch (InLocation)
@@ -801,11 +852,11 @@ void FAnimNextModuleInstance::QueueTask(FName InModuleEventName, TUniqueFunction
 	}
 	else
 	{
-		UE_LOGFMT(LogAnimation, Warning, "QueueTask: Could not find event '{EventName}' in module '{ModuleName}'", InModuleEventName, GetDataInterfaceName());
+		UE_LOGFMT(LogAnimation, Warning, "QueueTask: Could not find event '{EventName}' in module '{ModuleName}'", InEventName, GetDataInterfaceName());
 	}
 }
 
-void FAnimNextModuleInstance::QueueTaskOnOtherModule(const UE::AnimNext::FModuleHandle InOtherModuleHandle, FName InModuleEventName, TUniqueFunction<void(const UE::AnimNext::FModuleTaskContext&)>&& InTaskFunction, UE::AnimNext::ETaskRunLocation InLocation)
+void FAnimNextModuleInstance::QueueTaskOnOtherModule(const UE::AnimNext::FModuleHandle InOtherModuleHandle, const FName InEventName, TUniqueFunction<void(const UE::AnimNext::FModuleTaskContext&)>&& InTaskFunction, UE::AnimNext::ETaskRunLocation InLocation)
 {
 	if (Pool == nullptr)
 	{
@@ -818,5 +869,5 @@ void FAnimNextModuleInstance::QueueTaskOnOtherModule(const UE::AnimNext::FModule
 		return;
 	}
 
-	OtherModuleInstance->QueueTask(InModuleEventName, MoveTemp(InTaskFunction), InLocation);
+	OtherModuleInstance->QueueTask(InEventName, MoveTemp(InTaskFunction), InLocation);
 }
\ No newline at end of file
diff --git a/Engine/Plugins/Experimental/AnimNext/Source/AnimNext/Private/Module/ModuleTickFunction.cpp b/Engine/Plugins/Experimental/AnimNext/Source/AnimNext/Private/Module/ModuleTickFunction.cpp
--- a/Engine/Plugins/Experimental/AnimNext/Source/AnimNext/Private/Module/ModuleTickFunction.cpp
+++ b/Engine/Plugins/Experimental/AnimNext/Source/AnimNext/Private/Module/ModuleTickFunction.cpp
@@ -19,12 +19,12 @@
 namespace UE::AnimNext
 {
 
-void FModuleEndTickFunction::ExecuteTick(float DeltaTime, ELevelTick TickType, ENamedThreads::Type CurrentThread, const FGraphEventRef& MyCompletionGraphEvent)
+void FModuleEventTickFunction::Initialize(float InDeltaTime)
 {
-	Run(DeltaTime);
+	ModuleInstance->RunRigVMEvent(FRigUnit_AnimNextInitializeEvent::EventName, InDeltaTime);
 }
 
-void FModuleEndTickFunction::Run(float DeltaTime)
+void FModuleEventTickFunction::EndTick(float DeltaTime)
 {
 	SCOPED_NAMED_EVENT(AnimNext_Module_EndTick, FColor::Orange);
 
@@ -66,8 +66,10 @@ void FModuleEndTickFunction::Run(float DeltaTime)
 		ModuleInstance->OutputEventList.Reset();
 	}
 
-	if(ModuleInstance->RunState == FAnimNextModuleInstance::ERunState::PendingInitialUpdate)
+	if(ModuleInstance->InitState == FAnimNextModuleInstance::EInitState::FirstUpdate)
 	{
+		ModuleInstance->TransitionToInitState(FAnimNextModuleInstance::EInitState::Initialized);
+
 		if(ModuleInstance->InitMethod == EAnimNextModuleInitMethod::InitializeAndPause
 #if WITH_EDITOR
 			|| (ModuleInstance->InitMethod == EAnimNextModuleInitMethod::InitializeAndPauseInEditor && ModuleInstance->WorldType == EWorldType::Editor)
@@ -78,8 +80,11 @@ void FModuleEndTickFunction::Run(float DeltaTime)
 			FAnimNextModuleInstance::RunTaskOnGameThread([ModuleHandle = ModuleInstance->Handle, WeakObject = TWeakObjectPtr<UObject>(ModuleInstance->Object)]()
 			{
 				check(IsInGameThread());
-				UAnimNextWorldSubsystem* WorldSubsystem = UAnimNextWorldSubsystem::Get(WeakObject.Get());
-				WorldSubsystem->EnableHandle(ModuleHandle, false);
+				if (UObject* Object = WeakObject.Get())
+				{
+					UAnimNextWorldSubsystem* WorldSubsystem = UAnimNextWorldSubsystem::Get(Object);
+					WorldSubsystem->EnableHandle(ModuleHandle, false);
+				}
 			});
 		}
 	}
@@ -103,11 +108,6 @@ void FModuleEndTickFunction::Run(float DeltaTime)
 #endif
 }
 
-FString FModuleEndTickFunction::DiagnosticMessage()
-{
-	return TEXT("AnimNext: ModuleEnd");
-}
-
 void FModuleEventTickFunction::ExecuteTick(float DeltaTime, ELevelTick TickType, ENamedThreads::Type CurrentThread, const FGraphEventRef& MyCompletionGraphEvent)
 {
 	Run(DeltaTime);
@@ -119,6 +119,13 @@ void FModuleEventTickFunction::Run(float InDeltaTime)
 
 	FModuleWriteGuard Guard(ModuleInstance);
 
+	// Run the pending initialize if required
+	if (bFirstUserEvent && ModuleInstance->InitState == FAnimNextModuleInstance::EInitState::PendingInitializeEvent)
+	{
+		Initialize(InDeltaTime);
+		ModuleInstance->TransitionToInitState(FAnimNextModuleInstance::EInitState::FirstUpdate);
+	}
+
 	while (!PreExecuteTasks.IsEmpty())
 	{
 		TOptional<TUniqueFunction<void(const FModuleTaskContext&)>> Function = PreExecuteTasks.Dequeue();
@@ -140,48 +147,22 @@ void FModuleEventTickFunction::Run(float InDeltaTime)
 		check(Function.IsSet());
 		Function.GetValue()(FModuleTaskContext(*ModuleInstance));
 	}
+
+	if (bLastUserEvent)
+	{
+		EndTick(InDeltaTime);
+	}
 }
 
 #if WITH_EDITOR
 
 void FModuleEventTickFunction::InitializeAndRunModule(FAnimNextModuleInstance& InModuleInstance)
 {
-	// Sort tick functions topologically
-	TArray<FModuleEventTickFunction*> TickFunctionPtrs;
-	TickFunctionPtrs.Reserve(InModuleInstance.TickFunctions.Num());
-	Algo::Transform(InModuleInstance.TickFunctions, TickFunctionPtrs, [](const FModuleEventTickFunction& InTickFunction){ return const_cast<FModuleEventTickFunction*>(&InTickFunction); });
-	TArray<FModuleEventTickFunction*> DependencyPtrs;
-	Algo::TopologicalSort(
-		TickFunctionPtrs,
-		[&DependencyPtrs, &InModuleInstance](FModuleEventTickFunction* InTickFunction)
-		{
-			DependencyPtrs.Reset();
-
-			// Add direct prereqs
-			for(const FTickPrerequisite& Prerequisite : InTickFunction->GetPrerequisites())
-			{
-				DependencyPtrs.Add(static_cast<FModuleEventTickFunction*>(Prerequisite.PrerequisiteTickFunction));
-			}
-
-			// Also add any tick functions in prior tick groups that we don't already have
-			for(FModuleEventTickFunction& TickFunction : InModuleInstance.TickFunctions)
-			{
-				if(TickFunction.TickGroup < InTickFunction->TickGroup && !DependencyPtrs.Contains(&TickFunction))
-				{
-					DependencyPtrs.Add(&TickFunction);
-				}
-			}
-
-			return DependencyPtrs;
-		},
-		Algo::ETopologicalSort::None);
-
 	// Run sorted tick functions
-	for(FModuleEventTickFunction* TickFunction : TickFunctionPtrs)
+	for(FModuleEventTickFunction& TickFunction : InModuleInstance.TickFunctions)
 	{
-		TickFunction->Run(0.0f);
+		TickFunction.Run(0.0f);
 	}
-	InModuleInstance.EndTickFunction.Run(0.0f);
 }
 
 #endif
diff --git a/Engine/Plugins/Experimental/AnimNext/Source/AnimNext/Private/Module/RigUnit_AnimNextModuleEventDependencies.cpp b/Engine/Plugins/Experimental/AnimNext/Source/AnimNext/Private/Module/RigUnit_AnimNextModuleEventDependencies.cpp
--- a/Engine/Plugins/Experimental/AnimNext/Source/AnimNext/Private/Module/RigUnit_AnimNextModuleEventDependencies.cpp
+++ b/Engine/Plugins/Experimental/AnimNext/Source/AnimNext/Private/Module/RigUnit_AnimNextModuleEventDependencies.cpp
@@ -1,6 +1,8 @@
 // Copyright Epic Games, Inc. All Rights Reserved.
 
 #include "Module/RigUnit_AnimNextModuleEventDependencies.h"
+
+#include "Module/AnimNextModuleInstance.h"
 #include "Module/RigVMTrait_ModuleEventDependency.h"
 #if WITH_EDITOR
 #include "RigVMModel/RigVMController.h"
@@ -44,7 +46,7 @@ static void ApplyDependenciesHelper(const FAnimNextExecuteContext& InContext, co
 	bool bTickFunctionFound = false; 
 	for(const FRigVMTrait_ModuleEventDependency* Dependency : Dependencies)
 	{
-		FTickFunction* TickFunction = ModuleInstance.FindTickFunction(Dependency->EventName);
+		FTickFunction* TickFunction = ModuleInstance.FindTickFunctionByName(Dependency->EventName);
 		TickFunctions.Add(TickFunction);
 		bTickFunctionFound |= (TickFunction != nullptr);
 	}
diff --git a/Engine/Plugins/Experimental/AnimNext/Source/AnimNext/Private/Module/RigUnit_AnimNextModuleEvents.cpp b/Engine/Plugins/Experimental/AnimNext/Source/AnimNext/Private/Module/RigUnit_AnimNextModuleEvents.cpp
--- a/Engine/Plugins/Experimental/AnimNext/Source/AnimNext/Private/Module/RigUnit_AnimNextModuleEvents.cpp
+++ b/Engine/Plugins/Experimental/AnimNext/Source/AnimNext/Private/Module/RigUnit_AnimNextModuleEvents.cpp
@@ -24,17 +24,6 @@ UE::AnimNext::FModuleEventBindingFunction FRigUnit_AnimNextExecuteBindings_GT::G
 {
 	return [](const UE::AnimNext::FTickFunctionBindingContext& InContext, FTickFunction& InTickFunction)
 	{
-		InTickFunction.TickGroup = ETickingGroup::TG_PrePhysics;
-		InTickFunction.bRunOnAnyThread = false;
-
-		// Add prerequisite on this for all subsequent tick functions (assumes tick functions are sorted by phase)
-		TArrayView<UE::AnimNext::FModuleEventTickFunction> TickFunctions = InContext.ModuleInstance.GetTickFunctions();
-		for(int32 TickFunctionIndex = InContext.EventIndex + 1; TickFunctionIndex <TickFunctions.Num(); ++TickFunctionIndex)
-		{
-			FTickFunction& TickFunction = TickFunctions[TickFunctionIndex];
-			TickFunction.AddPrerequisite(InContext.ModuleInstance.GetObject(), TickFunctions[InContext.EventIndex]);
-		}
-
 		const bool bHighPriority = UE::AnimNext::CVars::CVarHighPriorityAnimNextExecuteBindingsTick.GetValueOnGameThread();
 		if (InTickFunction.bHighPriority != bHighPriority)
 		{
@@ -51,17 +40,6 @@ UE::AnimNext::FModuleEventBindingFunction FRigUnit_AnimNextExecuteBindings_WT::G
 {
 	return [](const UE::AnimNext::FTickFunctionBindingContext& InContext, FTickFunction& InTickFunction)
 	{
-		InTickFunction.TickGroup = ETickingGroup::TG_PrePhysics;
-		InTickFunction.bRunOnAnyThread = true; 
-
-		// Add prerequisite on this for all subsequent tick functions (assumes tick functions are sorted by phase)
-		TArrayView<UE::AnimNext::FModuleEventTickFunction> TickFunctions = InContext.ModuleInstance.GetTickFunctions();
-		for(int32 TickFunctionIndex = InContext.EventIndex + 1; TickFunctionIndex <TickFunctions.Num(); ++TickFunctionIndex)
-		{
-			FTickFunction& TickFunction = TickFunctions[TickFunctionIndex];
-			TickFunction.AddPrerequisite(InContext.ModuleInstance.GetObject(), TickFunctions[InContext.EventIndex]);
-		}
-
 		const bool bHighPriority = UE::AnimNext::CVars::CVarHighPriorityAnimNextExecuteBindingsTick.GetValueOnGameThread();
 		if (InTickFunction.bHighPriority != bHighPriority)
 		{
@@ -72,63 +50,31 @@ UE::AnimNext::FModuleEventBindingFunction FRigUnit_AnimNextExecuteBindings_WT::G
 
 FRigUnit_AnimNextInitializeEvent_Execute()
 {
-	FAnimNextModuleInstance& ModuleInstance = ExecuteContext.GetContextData<FAnimNextModuleContextData>().GetModuleInstance();
-
-	// Queue a GT task to disable ourselves (we can only enable/disable tick functions on the GT)
-	FAnimNextModuleInstance::RunTaskOnGameThread([ModuleHandle = ModuleInstance.Handle, WeakObject = TWeakObjectPtr<UObject>(ModuleInstance.Object)]()
-	{
-		UObject* Object = WeakObject.Get();
-		if (Object == nullptr)
-		{
-			return;
-		}
-		UAnimNextWorldSubsystem* WorldSubsystem = UAnimNextWorldSubsystem::Get(Object);
-		if (WorldSubsystem == nullptr)
-		{
-			return;
-		}
-		WorldSubsystem->DisableInitializeEvent(ModuleHandle);
-	});
 }
 
-UE::AnimNext::FModuleEventBindingFunction FRigUnit_AnimNextInitializeEvent::GetBindingFunction() const
+FString FRigUnit_AnimNextUserEvent::GetUnitSubTitle() const
 {
-	return [](const UE::AnimNext::FTickFunctionBindingContext& InContext, FTickFunction& InTickFunction)
+	FString Subtitle = UEnum::GetDisplayValueAsText<ETickingGroup>(TickGroup).ToString();
+	if (SortOrder != 0)
 	{
-		InContext.ModuleInstance.InitializeEventIndex = InContext.EventIndex;
-
-		InTickFunction.TickGroup = ETickingGroup::TG_PrePhysics;
+		Subtitle.Appendf(TEXT(" (%d)"), SortOrder);
+	}
 
-		// Add prerequisite on this for all subsequent tick functions (assumes tick functions are sorted by phase)
-		TArrayView<UE::AnimNext::FModuleEventTickFunction> TickFunctions = InContext.ModuleInstance.GetTickFunctions();
-		for(int32 TickFunctionIndex = InContext.EventIndex + 1; TickFunctionIndex < TickFunctions.Num(); ++TickFunctionIndex)
-		{
-			FTickFunction& TickFunction = TickFunctions[TickFunctionIndex];
-			TickFunction.AddPrerequisite(InContext.ModuleInstance.GetObject(), TickFunctions[InContext.EventIndex]);
-		}
-	};
-}
-
-FRigUnit_AnimNextPrePhysicsEvent_Execute()
-{
+	return Subtitle;
 }
 
-UE::AnimNext::FModuleEventBindingFunction FRigUnit_AnimNextPrePhysicsEvent::GetBindingFunction() const
+UE::AnimNext::FModuleEventBindingFunction FRigUnit_AnimNextUserEvent::GetBindingFunction() const
 {
-	return [](const UE::AnimNext::FTickFunctionBindingContext& InContext, FTickFunction& InTickFunction)
+	return [TickGroup = TickGroup](const UE::AnimNext::FTickFunctionBindingContext& InContext, FTickFunction& InTickFunction)
 	{
-		InTickFunction.TickGroup = ETickingGroup::TG_PrePhysics;
+		InTickFunction.TickGroup = TickGroup;
 	};
 }
 
-FRigUnit_AnimNextPostPhysicsEvent_Execute()
+FRigUnit_AnimNextPrePhysicsEvent_Execute()
 {
 }
 
-UE::AnimNext::FModuleEventBindingFunction FRigUnit_AnimNextPostPhysicsEvent::GetBindingFunction() const
+FRigUnit_AnimNextPostPhysicsEvent_Execute()
 {
-	return [](const UE::AnimNext::FTickFunctionBindingContext& InContext, FTickFunction& InTickFunction)
-	{
-		InTickFunction.TickGroup = ETickingGroup::TG_PostPhysics;
-	};
-}
\ No newline at end of file
+}
diff --git a/Engine/Plugins/Experimental/AnimNext/Source/AnimNext/Public/Module/RigVMTrait_ModuleEventDependency.h b/Engine/Plugins/Experimental/AnimNext/Source/AnimNext/Public/Module/RigVMTrait_ModuleEventDependency.h
--- a/Engine/Plugins/Experimental/AnimNext/Source/AnimNext/Public/Module/RigVMTrait_ModuleEventDependency.h
+++ b/Engine/Plugins/Experimental/AnimNext/Source/AnimNext/Public/Module/RigVMTrait_ModuleEventDependency.h
@@ -50,11 +50,11 @@ struct ANIMNEXT_API FRigVMTrait_ModuleEventDependency : public FRigVMTrait
 	// Override point that removes the dependency
 	virtual void OnRemoveDependency(const UE::AnimNext::FModuleDependencyContext& InContext) const {}
 
-	// The event to execute relative to
-	UPROPERTY(EditAnywhere, Category="Dependency", meta=(GetOptions="/Script/AnimNext.AnimNextConfig:GetModuleEventNames"))
-	FName EventName = FRigUnit_AnimNextPrePhysicsEvent::EventName;
-
 	// How to execute relative to the event
 	UPROPERTY(EditAnywhere, Category="Dependency")
 	EAnimNextModuleEventDependencyOrdering Ordering = EAnimNextModuleEventDependencyOrdering::Before;
-};
\ No newline at end of file
+
+	// The event to execute relative to
+	UPROPERTY(EditAnywhere, Category="Dependency", meta=(CustomWidget=AnimNextModuleEvent))
+	FName EventName = FRigUnit_AnimNextPrePhysicsEvent::DefaultEventName;
+};
diff --git a/Engine/Plugins/Experimental/AnimNext/Source/AnimNextEditor/Internal/Common/AnimNextAssetItemDetails.cpp b/Engine/Plugins/Experimental/AnimNext/Source/AnimNextEditor/Internal/Common/AnimNextAssetItemDetails.cpp
--- a/Engine/Plugins/Experimental/AnimNext/Source/AnimNextEditor/Internal/Common/AnimNextAssetItemDetails.cpp
+++ b/Engine/Plugins/Experimental/AnimNext/Source/AnimNextEditor/Internal/Common/AnimNextAssetItemDetails.cpp
@@ -115,12 +115,12 @@ void FAnimNextAssetItemDetails::RegisterToolMenuExtensions()
 
 						static FTextFormat AddEventTooltipFormat(LOCTEXT("AddEventGraphTooltipFormat", "Adds a {0} event graph to this asset"));
 
-						auto AddEventMenuEntry = [&EventsSection, SubEntryClass, WorkspaceItemContext, EditorData, Asset](FName InEventName, const FText& InDisplayName, UScriptStruct* InStruct)
+						auto AddEventMenuEntry = [&EventsSection, SubEntryClass, WorkspaceItemContext, EditorData, Asset](FName InEventName, const FText& InLabel, const FText& InTooltip, UScriptStruct* InStruct)
 						{
 							EventsSection.AddMenuEntry(
 								InEventName,
-								FText::FromName(InEventName),
-								FText::Format(AddEventTooltipFormat, InDisplayName),
+								InLabel,
+								InTooltip,
 								FSlateIconFinder::FindIconForClass(SubEntryClass, "ClassIcon.Object"),
 								FUIAction(
 									FExecuteAction::CreateWeakLambda(WorkspaceItemContext, [InEventName, SubEntryClass, EditorData, InStruct]()
@@ -152,11 +152,12 @@ void FAnimNextAssetItemDetails::RegisterToolMenuExtensions()
 								TInstancedStruct<FRigUnit_AnimNextModuleEventBase> StructInstance;
 								StructInstance.InitializeAsScriptStruct(Struct);
 								FName EventName = StructInstance.Get<FRigUnit_AnimNextModuleEventBase>().GetEventName();
-								AddEventMenuEntry(EventName, FText::FromName(EventName), Struct);
+								FText EventNameText = FText::FromName(EventName);
+								AddEventMenuEntry(EventName, EventNameText, FText::Format(AddEventTooltipFormat, EventNameText), Struct);
 							}
 						}
 
-						AddEventMenuEntry(TEXT("CustomEvent"), LOCTEXT("CustomEventDisplayNameInline", "custom"), FRigVMFunction_UserDefinedEvent::StaticStruct());
+						AddEventMenuEntry(TEXT("CustomEvent"), LOCTEXT("CustomEventLabel", "Custom Event"), FText::Format(AddEventTooltipFormat, LOCTEXT("CustomEventDisplayNameInline", "custom")), FRigVMFunction_UserDefinedEvent::StaticStruct());
 
 						EventsSection.AddSubMenu(
 							TEXT("NotifyEvents"),
diff --git a/Engine/Plugins/Experimental/AnimNext/Source/AnimNextEditor/Private/AnimNextEditorModule.cpp b/Engine/Plugins/Experimental/AnimNext/Source/AnimNextEditor/Private/AnimNextEditorModule.cpp
--- a/Engine/Plugins/Experimental/AnimNext/Source/AnimNextEditor/Private/AnimNextEditorModule.cpp
+++ b/Engine/Plugins/Experimental/AnimNext/Source/AnimNextEditor/Private/AnimNextEditorModule.cpp
@@ -70,6 +70,8 @@
 #include "Developer/TraceServices/Public/TraceServices/ModuleService.h"
 #include "Editor/RewindDebuggerInterface/Public/IRewindDebuggerExtension.h"
 #include "Editor/RewindDebuggerInterface/Public/IRewindDebuggerTrackCreator.h"
+#include "Graph/AnimNextGraphPanelPinFactory.h"
+#include "Module/ModuleEventPropertyCustomization.h"
 #include "Widgets/Docking/SDockTab.h"
 #include "RewindDebugger/RewindDebuggerAnimNext.h"
 #include "RewindDebugger/AnimNextModuleTrack.h"
@@ -115,6 +117,12 @@ void FAnimNextEditorModule::StartupModule()
 		"AnimNextVariableBinding",
 		FOnGetPropertyTypeCustomizationInstance::CreateLambda([] { return MakeShared<FVariableBindingPropertyCustomization>(); }));
 
+	ModuleEventPropertyTypeIdentifier = MakeShared<FModuleEventPropertyTypeIdentifier>();
+	PropertyModule.RegisterCustomPropertyTypeLayout(
+		"NameProperty",
+		FOnGetPropertyTypeCustomizationInstance::CreateLambda([] { return MakeShared<FModuleEventPropertyCustomization>(); }),
+		ModuleEventPropertyTypeIdentifier);
+
 	PropertyModule.RegisterCustomClassLayout("AnimNextVariableEntry", 
 		FOnGetDetailCustomizationInstance::CreateLambda([] { return MakeShared<FVariableCustomization>(); }));
 
@@ -124,6 +132,9 @@ void FAnimNextEditorModule::StartupModule()
 	PropertyModule.RegisterCustomClassLayout("AnimNextRigVMAssetEditorData",
 		FOnGetDetailCustomizationInstance::CreateLambda([] { return MakeShared<FAnimNextRigVMAssetEditorDataCustomization>(); }));
 
+	GraphPanelPinFactory = MakeShared<FAnimNextGraphPanelPinFactory>();
+	FEdGraphUtilities::RegisterVisualPinFactory(GraphPanelPinFactory);
+
 	Workspace::IWorkspaceEditorModule& WorkspaceEditorModule = FModuleManager::Get().LoadModuleChecked<Workspace::IWorkspaceEditorModule>("WorkspaceEditor");
 
 	WorkspaceEditorModule.OnRegisterTabsForEditor().AddLambda([](FWorkflowAllowedTabSet& TabFactories, const TSharedRef<FTabManager>& InTabManager, TSharedPtr<UE::Workspace::IWorkspaceEditor> InEditorPtr)
@@ -363,6 +374,7 @@ void FAnimNextEditorModule::ShutdownModule()
 		FPropertyEditorModule& PropertyModule = FModuleManager::LoadModuleChecked<FPropertyEditorModule>("PropertyEditor");
 		PropertyModule.UnregisterCustomPropertyTypeLayout("AnimNextParamType");
 		PropertyModule.UnregisterCustomPropertyTypeLayout("AnimNextVariableBinding");
+		PropertyModule.UnregisterCustomPropertyTypeLayout("NameProperty", ModuleEventPropertyTypeIdentifier);
 		PropertyModule.UnregisterCustomClassLayout("AnimNextVariableEntry");
 		PropertyModule.UnregisterCustomClassLayout("AnimNextVariableEntryProxy");
 		PropertyModule.UnregisterCustomClassLayout("AnimNextRigVMAssetEditorData");
@@ -413,6 +425,8 @@ void FAnimNextEditorModule::ShutdownModule()
 	IModularFeatures::Get().UnregisterModularFeature(IRewindDebuggerExtension::ModularFeatureName, &GRewindDebuggerAnimNext);
 	IModularFeatures::Get().UnregisterModularFeature(RewindDebugger::IRewindDebuggerTrackCreator::ModularFeatureName, &GAnimNextModulesTrackCreator);
 	IModularFeatures::Get().UnregisterModularFeature(TraceServices::ModuleFeatureName, &GAnimNextTraceModule);
+
+	FEdGraphUtilities::UnregisterVisualPinFactory(GraphPanelPinFactory);
 }
 
 void FAnimNextEditorModule::RegisterLocatorFragmentEditorType(FName InLocatorFragmentEditorName)
diff --git a/Engine/Plugins/Experimental/AnimNext/Source/AnimNextEditor/Private/AnimNextEditorModule.h b/Engine/Plugins/Experimental/AnimNext/Source/AnimNextEditor/Private/AnimNextEditorModule.h
--- a/Engine/Plugins/Experimental/AnimNext/Source/AnimNextEditor/Private/AnimNextEditorModule.h
+++ b/Engine/Plugins/Experimental/AnimNext/Source/AnimNextEditor/Private/AnimNextEditorModule.h
@@ -3,11 +3,13 @@
 #pragma once
 
 #include "IAnimNextEditorModule.h"
+#include "Module/ModuleEventPropertyCustomization.h"
 #include "UObject/TopLevelAssetPath.h"
 
 class UAnimNextWorkspaceSchema;
 class FAnimNextGraphPanelNodeFactory;
 class UAnimNextWorkspaceEditorMode;
+class FAnimNextGraphPanelPinFactory;
 
 namespace UE::AnimNext::Editor
 {
@@ -66,6 +68,9 @@ private:
 
 	TMap<FTopLevelAssetPath, FAssetCompilationHandlerFactoryDelegate> AssetCompilationHandlerFactories;
 
+	TSharedPtr<FAnimNextGraphPanelPinFactory> GraphPanelPinFactory;
+	TSharedPtr<FModuleEventPropertyTypeIdentifier> ModuleEventPropertyTypeIdentifier;
+
 	friend class FLocatorContext;
 };
 
diff --git a/Engine/Plugins/Experimental/AnimNext/Source/AnimNextEditor/Private/Graph/AnimNextGraphPanelPinFactory.cpp b/Engine/Plugins/Experimental/AnimNext/Source/AnimNextEditor/Private/Graph/AnimNextGraphPanelPinFactory.cpp
new file mode 100644
--- /dev/null
+++ b/Engine/Plugins/Experimental/AnimNext/Source/AnimNextEditor/Private/Graph/AnimNextGraphPanelPinFactory.cpp
@@ -0,0 +1,27 @@
+﻿// Copyright Epic Games, Inc. All Rights Reserved.
+
+#include "AnimNextGraphPanelPinFactory.h"
+
+#include "NodeFactory.h"
+#include "SGraphPinModuleEvent.h"
+#include "EdGraph/RigVMEdGraphNode.h"
+#include "Widgets/DeclarativeSyntaxSupport.h"
+
+TSharedPtr<SGraphPin> FAnimNextGraphPanelPinFactory::CreatePin(UEdGraphPin* InPin) const
+{
+	if (URigVMEdGraphNode* RigNode = Cast<URigVMEdGraphNode>(InPin->GetOwningNode()))
+	{
+		URigVMPin* ModelPin = RigNode->GetModelPinFromPinPath(InPin->GetName());
+		if (ModelPin)
+		{
+			static const FString META_AnimNextModuleEvent("AnimNextModuleEvent");
+			static const FName META_CustomWidget("CustomWidget");
+			if(InPin->PinType.PinCategory == UEdGraphSchema_K2::PC_Name && ModelPin->GetMetaData(META_CustomWidget) == META_AnimNextModuleEvent)
+			{
+				return SNew(UE::AnimNext::Editor::SGraphPinModuleEvent, InPin);
+			}
+		}
+	}
+
+	return nullptr;
+}
diff --git a/Engine/Plugins/Experimental/AnimNext/Source/AnimNextEditor/Private/Graph/AnimNextGraphPanelPinFactory.h b/Engine/Plugins/Experimental/AnimNext/Source/AnimNextEditor/Private/Graph/AnimNextGraphPanelPinFactory.h
new file mode 100644
--- /dev/null
+++ b/Engine/Plugins/Experimental/AnimNext/Source/AnimNextEditor/Private/Graph/AnimNextGraphPanelPinFactory.h
@@ -0,0 +1,13 @@
+﻿// Copyright Epic Games, Inc. All Rights Reserved.
+
+#pragma once
+
+#include "CoreMinimal.h"
+#include "EdGraphUtilities.h"
+
+class ANIMNEXTEDITOR_API FAnimNextGraphPanelPinFactory : public FGraphPanelPinFactory
+{
+public:
+	// FGraphPanelPinFactory interface
+	virtual TSharedPtr<SGraphPin> CreatePin(UEdGraphPin* InPin) const override;
+};
diff --git a/Engine/Plugins/Experimental/AnimNext/Source/AnimNextEditor/Private/Graph/SGraphPinModuleEvent.cpp b/Engine/Plugins/Experimental/AnimNext/Source/AnimNextEditor/Private/Graph/SGraphPinModuleEvent.cpp
new file mode 100644
--- /dev/null
+++ b/Engine/Plugins/Experimental/AnimNext/Source/AnimNextEditor/Private/Graph/SGraphPinModuleEvent.cpp
@@ -0,0 +1,31 @@
+﻿// Copyright Epic Games, Inc. All Rights Reserved.
+
+#include "SGraphPinModuleEvent.h"
+#include "SModuleEventPicker.h"
+
+namespace UE::AnimNext::Editor
+{
+
+void SGraphPinModuleEvent::Construct(const FArguments& InArgs, UEdGraphPin* InGraphPinObj)
+{
+	SGraphPin::Construct(SGraphPin::FArguments(), InGraphPinObj);
+}
+
+TSharedRef<SWidget>	SGraphPinModuleEvent::GetDefaultValueWidget()
+{
+	TArray<UObject*> ContextObjects;
+	ContextObjects.Add(GraphPinObj->GetOwningNode());
+	return SNew(SModuleEventPicker)
+		.ContextObjects(ContextObjects)
+		.InitiallySelectedEvent(FName(GraphPinObj->GetDefaultAsString()))
+		.OnEventPicked_Lambda([this](FName InEventName)
+		{
+			GraphPinObj->GetSchema()->TrySetDefaultValue(*GraphPinObj, InEventName.ToString());
+		})
+		.OnGetSelectedEvent_Lambda([this]()
+		{
+			return FName(GraphPinObj->GetDefaultAsString());
+		});
+}
+
+}
diff --git a/Engine/Plugins/Experimental/AnimNext/Source/AnimNextEditor/Private/Graph/SGraphPinModuleEvent.h b/Engine/Plugins/Experimental/AnimNext/Source/AnimNextEditor/Private/Graph/SGraphPinModuleEvent.h
new file mode 100644
--- /dev/null
+++ b/Engine/Plugins/Experimental/AnimNext/Source/AnimNextEditor/Private/Graph/SGraphPinModuleEvent.h
@@ -0,0 +1,22 @@
+﻿// Copyright Epic Games, Inc. All Rights Reserved.
+
+#pragma once
+
+#include "SGraphPin.h"
+
+namespace UE::AnimNext::Editor
+{
+
+class SGraphPinModuleEvent : public SGraphPin
+{
+public:
+	SLATE_BEGIN_ARGS(SGraphPinModuleEvent) {}
+	SLATE_END_ARGS()
+
+	void Construct(const FArguments& InArgs, UEdGraphPin* InGraphPinObj);
+
+private:
+	virtual TSharedRef<SWidget>	GetDefaultValueWidget() override;
+};
+
+}
\ No newline at end of file
diff --git a/Engine/Plugins/Experimental/AnimNext/Source/AnimNextEditor/Private/Graph/SModuleEventPicker.cpp b/Engine/Plugins/Experimental/AnimNext/Source/AnimNextEditor/Private/Graph/SModuleEventPicker.cpp
new file mode 100644
--- /dev/null
+++ b/Engine/Plugins/Experimental/AnimNext/Source/AnimNextEditor/Private/Graph/SModuleEventPicker.cpp
@@ -0,0 +1,153 @@
+﻿// Copyright Epic Games, Inc. All Rights Reserved.
+
+#include "SModuleEventPicker.h"
+
+#include "DetailLayoutBuilder.h"
+#include "Widgets/Input/SComboBox.h"
+#include "Module/AnimNextModule.h"
+#include "Module/RigUnit_AnimNextModuleEvents.h"
+#include "UObject/UObjectIterator.h"
+
+namespace UE::AnimNext::Editor
+{
+
+void SModuleEventPicker::Construct(const FArguments& InArgs)
+{
+	Algo::Transform(InArgs._ContextObjects, ContextObjects, [](UObject* ContextObject){ return TWeakObjectPtr<UObject>(ContextObject); });
+
+	RefreshEntries();
+
+	TSharedPtr<FName> InitialItem;
+	for (const TSharedPtr<FName>& Item : EventNamesSource)
+	{
+		if (InArgs._InitiallySelectedEvent == *Item)
+		{
+			InitialItem = Item;
+			break;
+		}
+	}
+
+	ChildSlot
+	[
+		SNew(SComboBox<TSharedPtr<FName>>)
+		.ContentPadding(0.0f)
+		.OptionsSource(&EventNamesSource)
+		.InitiallySelectedItem(InitialItem)
+		.OnComboBoxOpening_Lambda([this]()
+		{
+			RefreshEntries();
+		})
+		.OnSelectionChanged_Lambda([OnEventPicked = InArgs._OnEventPicked](TSharedPtr<FName> InItem, ESelectInfo::Type InSelectInfo)
+		{
+			if(InItem.IsValid())
+			{
+				OnEventPicked.ExecuteIfBound(*InItem);
+			}
+		})
+		.OnGenerateWidget_Lambda([](TSharedPtr<FName> InItem)
+		{
+			return SNew(STextBlock)
+				.Font(IDetailLayoutBuilder::GetDetailFont())
+				.Text(FText::FromName(InItem.IsValid() ? *InItem : NAME_None));
+		})
+		.Content()
+		[
+			SNew(STextBlock)
+			.Font(IDetailLayoutBuilder::GetDetailFont())
+			.Text_Lambda([OnGetSelectedEvent = InArgs._OnGetSelectedEvent]()
+			{
+				if (OnGetSelectedEvent.IsBound())
+				{
+					return FText::FromName(OnGetSelectedEvent.Execute());
+				}
+				return FText::GetEmpty();
+			})
+		]
+	];
+}
+
+void SModuleEventPicker::RefreshEntries()
+{
+	static const FName META_Hidden("Hidden");
+
+	TArray<FName> EventNames;
+	EventNamesSource.Empty();
+	
+	// Add default event names
+	for (TObjectIterator<UScriptStruct> It; It; ++It)
+	{
+		const UScriptStruct* Struct = *It;
+		if (Struct->IsChildOf(FRigUnit_AnimNextModuleEventBase::StaticStruct()) &&
+			Struct != FRigUnit_AnimNextModuleEventBase::StaticStruct() &&
+			!Struct->HasMetaData(META_Hidden))
+		{
+			TInstancedStruct<FRigUnit_AnimNextModuleEventBase> StructInstance;
+			StructInstance.InitializeAsScriptStruct(Struct);
+			const FRigUnit_AnimNextModuleEventBase& Event = StructInstance.Get();
+			if (Event.IsTask() && Event.IsUserEvent() && !Event.GetEventName().IsNone())
+			{
+				EventNames.AddUnique(Event.GetEventName());
+			}
+		}
+	}
+
+	// Parse out implemented events from the bytecode of the outer modules
+	for (TWeakObjectPtr<UObject> WeakObject : ContextObjects)
+	{
+		UObject* Object = WeakObject.Get();
+		if (Object == nullptr)
+		{
+			continue;
+		}
+
+		UAnimNextModule* Module = Cast<UAnimNextModule>(Object);
+		if(Module == nullptr)
+		{
+			Module = Object->GetTypedOuter<UAnimNextModule>();
+		}
+
+		if (Module == nullptr)
+		{
+			continue;
+		}
+
+		const URigVM* VM = Module->GetRigVM();
+		const FRigVMByteCode& ByteCode = VM->GetByteCode();
+		const TArray<const FRigVMFunction*>& Functions = VM->GetFunctions();
+		FRigVMInstructionArray Instructions = ByteCode.GetInstructions();
+		for (int32 EntryIndex = 0; EntryIndex < ByteCode.NumEntries(); EntryIndex++)
+		{
+			const FRigVMByteCodeEntry& Entry = ByteCode.GetEntry(EntryIndex);
+			const FRigVMInstruction& Instruction = Instructions[Entry.InstructionIndex];
+			const FRigVMExecuteOp& Op = ByteCode.GetOpAt<FRigVMExecuteOp>(Instruction);
+			const FRigVMFunction* Function = Functions[Op.FunctionIndex];
+			check(Function != nullptr);
+
+			if (Function->Struct->IsChildOf(FRigUnit_AnimNextModuleEventBase::StaticStruct()) &&
+				Function->Struct != FRigUnit_AnimNextModuleEventBase::StaticStruct() &&
+				!Function->Struct->HasMetaData(META_Hidden))
+			{
+				TInstancedStruct<FRigUnit_AnimNextModuleEventBase> StructInstance;
+				StructInstance.InitializeAsScriptStruct(Function->Struct);
+				const FRigUnit_AnimNextModuleEventBase& Event = StructInstance.Get();
+				if (Event.IsTask() && Event.IsUserEvent() && !Event.GetEventName().IsNone())
+				{
+					EventNames.AddUnique(Entry.Name);
+				}
+			}
+		}
+	}
+
+	if(EventNames.Num() == 0)
+	{
+		return;
+	}
+
+	EventNamesSource.Reserve(EventNames.Num());
+	for (FName EventName : EventNames)
+	{
+		EventNamesSource.Add(MakeShared<FName>(EventName));
+	}
+}
+
+}
diff --git a/Engine/Plugins/Experimental/AnimNext/Source/AnimNextEditor/Private/Graph/SModuleEventPicker.h b/Engine/Plugins/Experimental/AnimNext/Source/AnimNextEditor/Private/Graph/SModuleEventPicker.h
new file mode 100644
--- /dev/null
+++ b/Engine/Plugins/Experimental/AnimNext/Source/AnimNextEditor/Private/Graph/SModuleEventPicker.h
@@ -0,0 +1,39 @@
+﻿// Copyright Epic Games, Inc. All Rights Reserved.
+
+#pragma once
+
+#include "Widgets/SCompoundWidget.h"
+
+namespace UE::AnimNext::Editor
+{
+
+class SModuleEventPicker : public SCompoundWidget
+{
+public:
+	using FOnEventPicked = TDelegate<void(FName InEventName)>;
+	using FOnGetSelectedEvent = TDelegate<FName()>;
+
+	SLATE_BEGIN_ARGS(SModuleEventPicker) {}
+
+	SLATE_ARGUMENT(TArray<UObject*>, ContextObjects)
+
+	SLATE_EVENT(FOnEventPicked, OnEventPicked)
+
+	SLATE_EVENT(FOnGetSelectedEvent, OnGetSelectedEvent)
+
+	SLATE_ARGUMENT(FName, InitiallySelectedEvent)
+
+	SLATE_END_ARGS()
+
+	void Construct(const FArguments& InArgs);
+
+private:
+	void RefreshEntries();
+
+	// Names to be displayed in the combo box
+	TArray<TSharedPtr<FName>> EventNamesSource;
+
+	TArray<TWeakObjectPtr<UObject>> ContextObjects;
+};
+
+}
diff --git a/Engine/Plugins/Experimental/AnimNext/Source/AnimNextEditor/Private/Module/AnimNextModuleFactory.cpp b/Engine/Plugins/Experimental/AnimNext/Source/AnimNextEditor/Private/Module/AnimNextModuleFactory.cpp
--- a/Engine/Plugins/Experimental/AnimNext/Source/AnimNextEditor/Private/Module/AnimNextModuleFactory.cpp
+++ b/Engine/Plugins/Experimental/AnimNext/Source/AnimNextEditor/Private/Module/AnimNextModuleFactory.cpp
@@ -34,7 +34,7 @@ UObject* UAnimNextModuleFactory::FactoryCreateNew(UClass* Class, UObject* InPare
 	EditorData->bUsesExternalPackages = false;
 	EditorData->Initialize(/*bRecompileVM*/false);
 
-	EditorData->AddEventGraph(FRigUnit_AnimNextPrePhysicsEvent::EventName, FRigUnit_AnimNextPrePhysicsEvent::StaticStruct());
+	EditorData->AddEventGraph(FRigUnit_AnimNextPrePhysicsEvent::DefaultEventName, FRigUnit_AnimNextPrePhysicsEvent::StaticStruct());
 
 	// Compile the initial skeleton
 	EditorData->RecompileVM();
diff --git a/Engine/Plugins/Experimental/AnimNext/Source/AnimNextEditor/Private/Module/ModuleEventPropertyCustomization.cpp b/Engine/Plugins/Experimental/AnimNext/Source/AnimNextEditor/Private/Module/ModuleEventPropertyCustomization.cpp
new file mode 100644
--- /dev/null
+++ b/Engine/Plugins/Experimental/AnimNext/Source/AnimNextEditor/Private/Module/ModuleEventPropertyCustomization.cpp
@@ -0,0 +1,83 @@
+﻿// Copyright Epic Games, Inc. All Rights Reserved.
+
+#include "ModuleEventPropertyCustomization.h"
+
+#include "DetailLayoutBuilder.h"
+#include "DetailWidgetRow.h"
+#include "IDetailChildrenBuilder.h"
+#include "IDetailPropertyRow.h"
+#include "PropertyHandle.h"
+#include "SNameComboBox.h"
+#include "Graph/SModuleEventPicker.h"
+#include "Module/AnimNextModule.h"
+#include "Module/AnimNextModule_EditorData.h"
+#include "Module/RigVMTrait_ModuleEventDependency.h"
+#include "UObject/UObjectIterator.h"
+
+#define LOCTEXT_NAMESPACE "ModuleEventPropertyCustomization"
+
+namespace UE::AnimNext::Editor
+{
+
+bool FModuleEventPropertyTypeIdentifier::IsPropertyTypeCustomized(const IPropertyHandle& InPropertyHandle) const
+{
+	static const FString META_AnimNextModuleEvent("AnimNextModuleEvent");
+	static const FName META_CustomWidget("CustomWidget");
+	return InPropertyHandle.GetMetaData(META_CustomWidget) == META_AnimNextModuleEvent;
+}
+
+void FModuleEventPropertyCustomization::CustomizeHeader(TSharedRef<IPropertyHandle> InPropertyHandle, FDetailWidgetRow& InHeaderRow, IPropertyTypeCustomizationUtils& InCustomizationUtils)
+{
+	TArray<UObject*> Objects;
+	InPropertyHandle->GetOuterObjects(Objects);
+
+	FName SelectedValue;
+	InPropertyHandle->GetValue(SelectedValue);
+	TWeakPtr<IPropertyHandle> WeakPropertyHandle = InPropertyHandle;
+
+	InHeaderRow
+	.NameContent()
+	[
+		InPropertyHandle->CreatePropertyNameWidget()
+	]
+	.ValueContent()
+	.VAlign(VAlign_Center)
+	[
+		SNew(SModuleEventPicker)
+		.ContextObjects(Objects)
+		.InitiallySelectedEvent(SelectedValue)
+		.OnEventPicked_Lambda([WeakPropertyHandle](FName InEventName)
+		{
+			if (TSharedPtr<IPropertyHandle> PropertyHandle = WeakPropertyHandle.Pin())
+			{
+				PropertyHandle->SetValue(InEventName);
+			}
+		})
+		.OnGetSelectedEvent_Lambda([WeakPropertyHandle]()
+		{
+			if (TSharedPtr<IPropertyHandle> PropertyHandle = WeakPropertyHandle.Pin())
+			{
+				FName Value;
+				if (PropertyHandle->GetValue(Value) == FPropertyAccess::MultipleValues)
+				{
+					return FName(*LOCTEXT("MultipleValues", "Multiple Values").ToString());
+				}
+				else
+				{
+					return Value;
+				}
+			}
+
+			return FName();
+		})
+	];
+}
+
+void FModuleEventPropertyCustomization::CustomizeChildren(TSharedRef<IPropertyHandle> InPropertyHandle, IDetailChildrenBuilder& ChildBuilder, IPropertyTypeCustomizationUtils& CustomizationUtils)
+{
+
+}
+
+}
+
+#undef LOCTEXT_NAMESPACE
\ No newline at end of file
diff --git a/Engine/Plugins/Experimental/AnimNext/Source/AnimNextEditor/Private/Module/ModuleEventPropertyCustomization.h b/Engine/Plugins/Experimental/AnimNext/Source/AnimNextEditor/Private/Module/ModuleEventPropertyCustomization.h
new file mode 100644
--- /dev/null
+++ b/Engine/Plugins/Experimental/AnimNext/Source/AnimNextEditor/Private/Module/ModuleEventPropertyCustomization.h
@@ -0,0 +1,26 @@
+﻿// Copyright Epic Games, Inc. All Rights Reserved.
+
+#pragma once
+
+#include "IPropertyTypeCustomization.h"
+#include "Param/ParamTypePropertyCustomization.h"
+
+namespace UE::AnimNext::Editor
+{
+
+class FModuleEventPropertyTypeIdentifier : public IPropertyTypeIdentifier
+{
+	virtual bool IsPropertyTypeCustomized(const IPropertyHandle& InPropertyHandle) const override;
+};
+
+class FModuleEventPropertyCustomization : public IPropertyTypeCustomization
+{
+private:
+	// IPropertyTypeCustomization interface
+	virtual void CustomizeHeader(TSharedRef<IPropertyHandle> InPropertyHandle, FDetailWidgetRow& HeaderRow, IPropertyTypeCustomizationUtils& CustomizationUtils) override;
+	virtual void CustomizeChildren(TSharedRef<IPropertyHandle> InPropertyHandle, IDetailChildrenBuilder& ChildBuilder, IPropertyTypeCustomizationUtils& CustomizationUtils) override;
+
+
+};
+
+}
diff --git a/Engine/Plugins/Experimental/AnimNext/Source/AnimNextTestSuite/Private/AnimNextEditorTest.cpp b/Engine/Plugins/Experimental/AnimNext/Source/AnimNextTestSuite/Private/AnimNextEditorTest.cpp
--- a/Engine/Plugins/Experimental/AnimNext/Source/AnimNextTestSuite/Private/AnimNextEditorTest.cpp
+++ b/Engine/Plugins/Experimental/AnimNext/Source/AnimNextTestSuite/Private/AnimNextEditorTest.cpp
@@ -8,6 +8,7 @@
 #include "Animation/AnimSequence.h"
 #include "Entries/AnimNextVariableEntry.h"
 #include "Entries/AnimNextEventGraphEntry.h"
+#include "Module/AnimNextModule.h"
 #include "Module/AnimNextModuleFactory.h"
 #include "Module/RigUnit_AnimNextModuleEvents.h"
 #if WITH_EDITOR
diff --git a/Engine/Plugins/Experimental/AnimNext/Source/AnimNextTestSuite/Private/AnimNextVariablesTest.cpp b/Engine/Plugins/Experimental/AnimNext/Source/AnimNextTestSuite/Private/AnimNextVariablesTest.cpp
--- a/Engine/Plugins/Experimental/AnimNext/Source/AnimNextTestSuite/Private/AnimNextVariablesTest.cpp
+++ b/Engine/Plugins/Experimental/AnimNext/Source/AnimNextTestSuite/Private/AnimNextVariablesTest.cpp
@@ -291,7 +291,7 @@ bool FVariables::RunTest(const FString& InParameters)
 	});
 	Asset->GetRigVMExtendedExecuteContext().SetRuntimeSettings(RuntimeSettings);
 
-	Asset->GetVM()->ExecuteVM(Asset->GetRigVMExtendedExecuteContext(), FRigUnit_AnimNextPrePhysicsEvent::EventName);
+	Asset->GetVM()->ExecuteVM(Asset->GetRigVMExtendedExecuteContext(), FRigUnit_AnimNextPrePhysicsEvent::DefaultEventName);
 
 	UE_RETURN_ON_ERROR(Messages.Num() == 1, TEXT("unexpected number of messages"));
 	UE_RETURN_ON_ERROR(Messages[0] == TEXT("Result = 3"), TEXT("unexpected result message"));
@@ -413,7 +413,7 @@ bool FVariables_UOLBindings::RunTest(const FString& InParameters)
 	Asset->GetVM()->ExecuteVM(Asset->GetRigVMExtendedExecuteContext(), FRigUnit_AnimNextExecuteBindings_GT::EventName);
 
 	// Run PrePhysics event to print messages
-	Asset->GetVM()->ExecuteVM(Asset->GetRigVMExtendedExecuteContext(), FRigUnit_AnimNextPrePhysicsEvent::EventName);
+	Asset->GetVM()->ExecuteVM(Asset->GetRigVMExtendedExecuteContext(), FRigUnit_AnimNextPrePhysicsEvent::DefaultEventName);
 
 	UE_RETURN_ON_ERROR(Messages.Num() == 3, TEXT("unexpected number of messages"));
 	UE_RETURN_ON_ERROR(Messages[0] == TEXT("Result = 23"), TEXT("unexpected result message"));
diff --git a/Engine/Plugins/Experimental/AnimNext/Source/AnimNextUncookedOnly/Internal/AnimNextRigVMAssetEditorData.h b/Engine/Plugins/Experimental/AnimNext/Source/AnimNextUncookedOnly/Internal/AnimNextRigVMAssetEditorData.h
--- a/Engine/Plugins/Experimental/AnimNext/Source/AnimNextUncookedOnly/Internal/AnimNextRigVMAssetEditorData.h
+++ b/Engine/Plugins/Experimental/AnimNext/Source/AnimNextUncookedOnly/Internal/AnimNextRigVMAssetEditorData.h
@@ -208,6 +208,7 @@ protected:
 	// UObject interface
 	virtual void Serialize(FArchive& Ar) override;
 	virtual void PostLoad() override;
+	virtual void PostEditChangeProperty(FPropertyChangedEvent& PropertyChangedEvent) override;
 	virtual void PostTransacted(const FTransactionObjectEvent& TransactionEvent) override;
 	virtual void GetAssetRegistryTags(FAssetRegistryTagsContext Context) const override;
 	virtual bool IsEditorOnly() const override { return true; }
diff --git a/Engine/Plugins/Experimental/AnimNext/Source/AnimNextUncookedOnly/Internal/Module/AnimNextModule_EditorData.h b/Engine/Plugins/Experimental/AnimNext/Source/AnimNextUncookedOnly/Internal/Module/AnimNextModule_EditorData.h
--- a/Engine/Plugins/Experimental/AnimNext/Source/AnimNextUncookedOnly/Internal/Module/AnimNextModule_EditorData.h
+++ b/Engine/Plugins/Experimental/AnimNext/Source/AnimNextUncookedOnly/Internal/Module/AnimNextModule_EditorData.h
@@ -7,6 +7,7 @@
 #include "AnimNextEdGraph.h"
 #include "DataInterface/AnimNextDataInterface_EditorData.h"
 #include "AnimNextExecuteContext.h"
+#include "Module/RigVMTrait_ModuleEventDependency.h"
 #include "AnimNextModule_EditorData.generated.h"
 
 class UAnimNextModule;
@@ -61,4 +62,8 @@ private:
 private:
 	UPROPERTY()
 	TArray<TObjectPtr<UAnimNextEdGraph>> Graphs_DEPRECATED;
+
+	// All dependencies that should be set up when the module initializes
+	UPROPERTY(EditAnywhere, Category = "Dependencies", NoClear, meta=(ExcludeBaseStruct))
+	TArray<TInstancedStruct<FRigVMTrait_ModuleEventDependency>> Dependencies;
 };
diff --git a/Engine/Plugins/Experimental/AnimNext/Source/AnimNextUncookedOnly/Private/AnimNextRigVMAssetEditorData.cpp b/Engine/Plugins/Experimental/AnimNext/Source/AnimNextUncookedOnly/Private/AnimNextRigVMAssetEditorData.cpp
--- a/Engine/Plugins/Experimental/AnimNext/Source/AnimNextUncookedOnly/Private/AnimNextRigVMAssetEditorData.cpp
+++ b/Engine/Plugins/Experimental/AnimNext/Source/AnimNextUncookedOnly/Private/AnimNextRigVMAssetEditorData.cpp
@@ -258,6 +258,13 @@ void UAnimNextRigVMAssetEditorData::PostLoadExternalPackages()
 	Entries.Append(InternalEntries);
 }
 
+void UAnimNextRigVMAssetEditorData::PostEditChangeProperty(FPropertyChangedEvent& PropertyChangedEvent)
+{
+	Super::PostEditChangeProperty(PropertyChangedEvent);
+
+	BroadcastModified(EAnimNextEditorDataNotifType::PropertyChanged, this);
+}
+
 void UAnimNextRigVMAssetEditorData::PostTransacted(const FTransactionObjectEvent& TransactionEvent)
 {
 	Super::PostTransacted(TransactionEvent);
@@ -1836,7 +1843,7 @@ UAnimNextEventGraphEntry* UAnimNextRigVMAssetEditorData::AddEventGraph(FName InN
 			RefreshExternalModels();
 			RigVMClient.AddModel(NewRigVMGraphModel, true);
 			URigVMController* Controller = RigVMClient.GetController(NewRigVMGraphModel);
-			UE::AnimNext::UncookedOnly::FUtils::SetupEventGraph(Controller, InEventStruct, InName);
+			UE::AnimNext::UncookedOnly::FUtils::SetupEventGraph(Controller, InEventStruct, NewGraphName);
 		}
 	}
 
diff --git a/Engine/Plugins/Experimental/AnimNext/Source/AnimNextUncookedOnly/Private/Module/AnimNextModule_EditorData.cpp b/Engine/Plugins/Experimental/AnimNext/Source/AnimNextUncookedOnly/Private/Module/AnimNextModule_EditorData.cpp
--- a/Engine/Plugins/Experimental/AnimNext/Source/AnimNextUncookedOnly/Private/Module/AnimNextModule_EditorData.cpp
+++ b/Engine/Plugins/Experimental/AnimNext/Source/AnimNextUncookedOnly/Private/Module/AnimNextModule_EditorData.cpp
@@ -51,6 +51,7 @@ void UAnimNextModule_EditorData::OnPreCompileAsset(FRigVMCompileSettings& InSett
 	UAnimNextModule* Module = FUtils::GetAsset<UAnimNextModule>(this);
 
 	Module->RequiredComponents.Empty();
+	Module->Dependencies.Empty();
 }
 
 void UAnimNextModule_EditorData::OnPreCompileGetProgrammaticGraphs(const FRigVMCompileSettings& InSettings, FAnimNextGetGraphCompileContext& OutCompileContext)
@@ -101,6 +102,9 @@ void UAnimNextModule_EditorData::OnPreCompileProcessGraphs(const FRigVMCompileSe
 			});
 		}
 	}
+
+	// Copy dependencies
+	Module->Dependencies = Dependencies;
 }
 
 void UAnimNextModule_EditorData::CustomizeNewAssetEntry(UAnimNextRigVMAssetEntry* InNewEntry) const
diff --git a/Engine/Plugins/Experimental/AnimNext/Source/AnimNextUncookedOnly/Private/UncookedOnlyUtils.cpp b/Engine/Plugins/Experimental/AnimNext/Source/AnimNextUncookedOnly/Private/UncookedOnlyUtils.cpp
--- a/Engine/Plugins/Experimental/AnimNext/Source/AnimNextUncookedOnly/Private/UncookedOnlyUtils.cpp
+++ b/Engine/Plugins/Experimental/AnimNext/Source/AnimNextUncookedOnly/Private/UncookedOnlyUtils.cpp
@@ -795,13 +795,22 @@ void FUtils::SetupEventGraph(URigVMController* InController, UScriptStruct* InEv
 	// Clear the graph
 	InController->RemoveNodes(InController->GetGraph()->GetNodes());
 
-	// Add entry point
-	URigVMUnitNode* EventNode = InController->AddUnitNode(InEventStruct, FRigVMStruct::ExecuteName, FVector2D(-200.0f, 0.0f), FString(), false, bPrintPythonCommand);
-
-	// Setup custom event name if required
-	if(EventNode && InEventStruct == FRigVMFunction_UserDefinedEvent::StaticStruct() && InEventName != NAME_None && !InController->GetAllEventNames().Contains(InEventName))
+	if (InEventStruct->IsChildOf(FRigUnit_AnimNextUserEvent::StaticStruct()))
+	{
+		FRigUnit_AnimNextUserEvent Defaults;
+		Defaults.Name = InEventName;
+		Defaults.SortOrder = InEventName.GetNumber();
+		InController->AddUnitNodeWithDefaults(InEventStruct, Defaults, FRigVMStruct::ExecuteName, FVector2D(-200.0f, 0.0f), FString(), false);
+	}
+	else if (InEventStruct == FRigVMFunction_UserDefinedEvent::StaticStruct())
+	{
+		FRigVMFunction_UserDefinedEvent Defaults;
+		Defaults.EventName = InEventName;
+		InController->AddUnitNodeWithDefaults(InEventStruct, Defaults, FRigVMStruct::ExecuteName, FVector2D(-200.0f, 0.0f), FString(), false);
+	}
+	else
 	{
-		InController->SetPinDefaultValue(EventNode->FindPin(GET_MEMBER_NAME_STRING_CHECKED(FRigVMFunction_UserDefinedEvent, EventName)), InEventName.ToString(), true, true, bPrintPythonCommand, false);
+		InController->AddUnitNode(InEventStruct, FRigVMStruct::ExecuteName, FVector2D(-200.0f, 0.0f), FString(), false);
 	}
 }
 
diff --git a/Engine/Plugins/Experimental/AnimNextAnimGraph/Source/AnimNextAnimGraph/Private/Graph/AnimNextAnimationGraph.cpp b/Engine/Plugins/Experimental/AnimNextAnimGraph/Source/AnimNextAnimGraph/Private/Graph/AnimNextAnimationGraph.cpp
--- a/Engine/Plugins/Experimental/AnimNextAnimGraph/Source/AnimNextAnimGraph/Private/Graph/AnimNextAnimationGraph.cpp
+++ b/Engine/Plugins/Experimental/AnimNextAnimGraph/Source/AnimNextAnimGraph/Private/Graph/AnimNextAnimationGraph.cpp
@@ -15,6 +15,7 @@
 #include "Graph/RigUnit_AnimNextGraphEvaluator.h"
 #include "UObject/LinkerLoad.h"
 #include "UObject/ObjectResource.h"
+#include "Module/AnimNextModuleInstance.h"
 
 #if WITH_EDITOR	
 #include "Engine/ExternalAssetDependencyGatherer.h"
diff --git a/Engine/Plugins/Experimental/AnimNextPoseSearch/Source/AnimNextPoseSearch/Private/AnimNextInteractionIslandDependency.cpp b/Engine/Plugins/Experimental/AnimNextPoseSearch/Source/AnimNextPoseSearch/Private/AnimNextInteractionIslandDependency.cpp
--- a/Engine/Plugins/Experimental/AnimNextPoseSearch/Source/AnimNextPoseSearch/Private/AnimNextInteractionIslandDependency.cpp
+++ b/Engine/Plugins/Experimental/AnimNextPoseSearch/Source/AnimNextPoseSearch/Private/AnimNextInteractionIslandDependency.cpp
@@ -18,31 +18,31 @@ bool FAnimNextInteractionIslandDependency::CanMakeDependency(const UObject* InIs
 const FTickFunction* FAnimNextInteractionIslandDependency::FindTickFunction(UObject* InObject) const
 {
 	const UAnimNextComponent* AnimNextComponent = CastChecked<UAnimNextComponent>(InObject);
-	return AnimNextComponent->FindTickFunction(FRigUnit_AnimNextPrePhysicsEvent::EventName);
+	return AnimNextComponent->FindTickFunction(FRigUnit_AnimNextPrePhysicsEvent::DefaultEventName);
 }
 
 void FAnimNextInteractionIslandDependency::AddPrerequisite(UObject* InIslandObject, FTickFunction& InIslandTickFunction, UObject* InObject) const
 {
 	UAnimNextComponent* AnimNextComponent = CastChecked<UAnimNextComponent>(InObject);
-	AnimNextComponent->AddSubsequent(InIslandObject, InIslandTickFunction, FRigUnit_AnimNextPrePhysicsEvent::EventName);
+	AnimNextComponent->AddSubsequent(InIslandObject, InIslandTickFunction, FRigUnit_AnimNextPrePhysicsEvent::DefaultEventName);
 }
 
 void FAnimNextInteractionIslandDependency::AddSubsequent(UObject* InIslandObject, FTickFunction& InIslandTickFunction, UObject* InObject) const
 {
 	UAnimNextComponent* AnimNextComponent = CastChecked<UAnimNextComponent>(InObject);
-	AnimNextComponent->AddPrerequisite(InIslandObject, InIslandTickFunction, FRigUnit_AnimNextPrePhysicsEvent::EventName);
+	AnimNextComponent->AddPrerequisite(InIslandObject, InIslandTickFunction, FRigUnit_AnimNextPrePhysicsEvent::DefaultEventName);
 }
 
 void FAnimNextInteractionIslandDependency::RemovePrerequisite(UObject* InIslandObject, FTickFunction& InIslandTickFunction, UObject* InObject) const
 {
 	UAnimNextComponent* AnimNextComponent = CastChecked<UAnimNextComponent>(InObject);
-	AnimNextComponent->RemoveSubsequent(InIslandObject, InIslandTickFunction, FRigUnit_AnimNextPrePhysicsEvent::EventName);
+	AnimNextComponent->RemoveSubsequent(InIslandObject, InIslandTickFunction, FRigUnit_AnimNextPrePhysicsEvent::DefaultEventName);
 }
 
 void FAnimNextInteractionIslandDependency::RemoveSubsequent(UObject* InIslandObject, FTickFunction& InIslandTickFunction, UObject* InObject) const
 {
 	UAnimNextComponent* AnimNextComponent = CastChecked<UAnimNextComponent>(InObject);
-	AnimNextComponent->RemovePrerequisite(InIslandObject, InIslandTickFunction, FRigUnit_AnimNextPrePhysicsEvent::EventName);
+	AnimNextComponent->RemovePrerequisite(InIslandObject, InIslandTickFunction, FRigUnit_AnimNextPrePhysicsEvent::DefaultEventName);
 }
 
 }
diff --git a/Engine/Plugins/Experimental/AnimNextPoseSearch/Source/AnimNextPoseSearch/Private/MotionMatchingTrait.cpp b/Engine/Plugins/Experimental/AnimNextPoseSearch/Source/AnimNextPoseSearch/Private/MotionMatchingTrait.cpp
--- a/Engine/Plugins/Experimental/AnimNextPoseSearch/Source/AnimNextPoseSearch/Private/MotionMatchingTrait.cpp
+++ b/Engine/Plugins/Experimental/AnimNextPoseSearch/Source/AnimNextPoseSearch/Private/MotionMatchingTrait.cpp
@@ -472,7 +472,7 @@ namespace UE::AnimNext
 		{
 			ModuleInstance->QueueTaskOnOtherModule(
 				ResultModuleHandle.ModuleHandle,
-				FRigUnit_AnimNextPrePhysicsEvent::EventName,
+				FRigUnit_AnimNextPrePhysicsEvent::DefaultEventName,
 				[Result, ResultOutputVariableName](const UE::AnimNext::FModuleTaskContext& TaskContext)
 				{
 					FAnimNextModuleInstance* OtherModuleInstance = TaskContext.GetModuleInstance();
diff --git a/Engine/Plugins/Experimental/AnimNextPoseSearch/Source/AnimNextPoseSearch/Private/RigUnit_DebugDrawTrajectory.cpp b/Engine/Plugins/Experimental/AnimNextPoseSearch/Source/AnimNextPoseSearch/Private/RigUnit_DebugDrawTrajectory.cpp
--- a/Engine/Plugins/Experimental/AnimNextPoseSearch/Source/AnimNextPoseSearch/Private/RigUnit_DebugDrawTrajectory.cpp
+++ b/Engine/Plugins/Experimental/AnimNextPoseSearch/Source/AnimNextPoseSearch/Private/RigUnit_DebugDrawTrajectory.cpp
@@ -5,6 +5,7 @@
 #include "Animation/TrajectoryTypes.h"
 #include "Component/AnimNextComponent.h"
 #include "PoseSearch/PoseSearchDefines.h"
+#include "Module/AnimNextModuleInstance.h"
 
 FRigUnit_DebugDrawTrajectory_Execute()
 {
diff --git a/Engine/Plugins/Experimental/MoverAnimNext/Source/MoverAnimNext/Private/Graph/RigUnit_GenerateMoverTrajectory.cpp b/Engine/Plugins/Experimental/MoverAnimNext/Source/MoverAnimNext/Private/Graph/RigUnit_GenerateMoverTrajectory.cpp
--- a/Engine/Plugins/Experimental/MoverAnimNext/Source/MoverAnimNext/Private/Graph/RigUnit_GenerateMoverTrajectory.cpp
+++ b/Engine/Plugins/Experimental/MoverAnimNext/Source/MoverAnimNext/Private/Graph/RigUnit_GenerateMoverTrajectory.cpp
@@ -12,13 +12,6 @@ FRigUnit_GenerateMoverTrajectory_Execute()
 	TRACE_CPUPROFILER_EVENT_SCOPE(FRigUnit_GenerateMoverTrajectory_Execute);
 	using namespace UE::AnimNext;
 
-	const FAnimNextModuleContextData& ModuleContextData = ExecuteContext.GetContextData<FAnimNextModuleContextData>();
-	FAnimNextModuleInstance& ModuleInstance = ModuleContextData.GetModuleInstance();
-
-	// TODO: this is not thread safe at the moment. Once we get the 'weak semantics' CL checked in, we can move commonly-used debug info
-	// into a module component that can be used to access GT data where needed. 
-	UAnimNextComponent* Component = Cast<UAnimNextComponent>(ModuleInstance.GetObject());
-
 	// Note this logic is based on UPoseSearchTrajectoryLibrary::PoseSearchGeneratePredictorTrajectory
 
 	FPoseSearchTrajectoryData::FSampling TrajectoryDataSampling;
diff --git a/Engine/Plugins/Runtime/RigVM/Source/RigVMDeveloper/Private/RigVMCompiler/RigVMCompiler.cpp b/Engine/Plugins/Runtime/RigVM/Source/RigVMDeveloper/Private/RigVMCompiler/RigVMCompiler.cpp
--- a/Engine/Plugins/Runtime/RigVM/Source/RigVMDeveloper/Private/RigVMCompiler/RigVMCompiler.cpp
+++ b/Engine/Plugins/Runtime/RigVM/Source/RigVMDeveloper/Private/RigVMCompiler/RigVMCompiler.cpp
@@ -651,7 +651,8 @@ bool URigVMCompiler::Compile(const FRigVMCompileSettings& InSettings, TArray<URi
 			// that is, make sure their cpp type and type object agree with each other
 			VisitedGraphController->EnsureLocalVariableValidity();
 		}
-		
+
+		TSet<FName> FoundEvents;
 		for(URigVMNode* ModelNode : VisitedGraph->GetNodes())
 		{
 			URigVMController* VisitedGraphController = InController->GetControllerForGraph(VisitedGraph);
@@ -672,6 +673,17 @@ bool URigVMCompiler::Compile(const FRigVMCompileSettings& InSettings, TArray<URi
 				}
 			}
 
+			if (ModelNode->IsEvent())
+			{
+				if (FoundEvents.Contains(ModelNode->GetEventName()))
+				{
+					static const FString LinkedMessage = TEXT("Duplicate event node @@ found.");
+					Settings.ASTSettings.Report(EMessageSeverity::Warning, ModelNode, LinkedMessage);
+				}
+
+				FoundEvents.Add(ModelNode->GetEventName());
+			}
+
 			if(ModelNode->IsA<UDEPRECATED_RigVMBranchNode>())
 			{
 				static const FString LinkedMessage = TEXT("Node @@ is a deprecated branch node. Cannot compile.");
