AnimNext: Fix uninitialized variable in AnimNext event

diff --git a/Engine/Plugins/Experimental/AnimNext/Source/AnimNext/Internal/Module/RigUnit_AnimNextModuleEvents.h b/Engine/Plugins/Experimental/AnimNext/Source/AnimNext/Internal/Module/RigUnit_AnimNextModuleEvents.h
--- a/Engine/Plugins/Experimental/AnimNext/Source/AnimNext/Internal/Module/RigUnit_AnimNextModuleEvents.h
+++ b/Engine/Plugins/Experimental/AnimNext/Source/AnimNext/Internal/Module/RigUnit_AnimNextModuleEvents.h
@@ -143,7 +143,7 @@ struct ANIMNEXT_API FRigUnit_AnimNextUserEvent : public FRigUnit_AnimNextModuleE
 
 	// The tick group the event executes in 
 	UPROPERTY()
-	TEnumAsByte<ETickingGroup> TickGroup;
+	TEnumAsByte<ETickingGroup> TickGroup = ETickingGroup::TG_PrePhysics;
 };
 
 /** Schedule event called before world physics is updated */
