Avoid weak object ptr testing in hot loop

diff --git a/Engine/Plugins/Experimental/AnimNext/Source/AnimNext/Private/DecompressionTools.cpp b/Engine/Plugins/Experimental/AnimNext/Source/AnimNext/Private/DecompressionTools.cpp
--- a/Engine/Plugins/Experimental/AnimNext/Source/AnimNext/Private/DecompressionTools.cpp
+++ b/Engine/Plugins/Experimental/AnimNext/Source/AnimNext/Private/DecompressionTools.cpp
@@ -330,6 +330,7 @@ void FDecompressionTools::DecompressPose(FLODPose& OutAnimationPoseData,
 	const USkeleton* SourceSkeleton = DecompressionContext.GetSourceSkeleton();
 	const USkeleton* TargetSkeleton = OutAnimationPoseData.GetSkeletonAsset();
 	const FSkeletonRemapping& SkeletonRemapping = UE::Anim::FSkeletonRemappingRegistry::Get().GetRemapping(SourceSkeleton, TargetSkeleton);
+	const bool bIsSkeletonRemappingValid = SkeletonRemapping.IsValid();
 
 	const bool bUseSourceRetargetModes = TargetSkeleton->GetUseRetargetModesFromCompatibleSkeleton();
 	const bool bDisableRetargeting = OutAnimationPoseData.GetDisableRetargeting();
@@ -372,7 +373,7 @@ void FDecompressionTools::DecompressPose(FLODPose& OutAnimationPoseData,
 		for (int32 TrackIndex = (bFirstTrackIsRootBone ? 1 : 0); TrackIndex < NumTracks; TrackIndex++)
 		{
 			const int32 SourceSkeletonBoneIndex = CompressedData.GetSkeletonIndexFromTrackIndex(TrackIndex);
-			const int32 TargetSkeletonBoneIndex = SkeletonRemapping.IsValid() ? SkeletonRemapping.GetTargetSkeletonBoneIndex(SourceSkeletonBoneIndex) : SourceSkeletonBoneIndex;
+			const int32 TargetSkeletonBoneIndex = bIsSkeletonRemappingValid ? SkeletonRemapping.GetTargetSkeletonBoneIndex(SourceSkeletonBoneIndex) : SourceSkeletonBoneIndex;
 
 			if (TargetSkeletonBoneIndex != INDEX_NONE)
 			{
diff --git a/Engine/Source/Runtime/Engine/Private/Animation/AnimationDecompression.cpp b/Engine/Source/Runtime/Engine/Private/Animation/AnimationDecompression.cpp
--- a/Engine/Source/Runtime/Engine/Private/Animation/AnimationDecompression.cpp
+++ b/Engine/Source/Runtime/Engine/Private/Animation/AnimationDecompression.cpp
@@ -53,6 +53,7 @@ void DecompressPose(FCompactPose& OutPose,
 	const USkeleton* SourceSkeleton = DecompressionContext.GetSourceSkeleton();
 	const USkeleton* TargetSkeleton = RequiredBones.GetSkeletonAsset();
 	const FSkeletonRemapping& SkeletonRemapping = UE::Anim::FSkeletonRemappingRegistry::Get().GetRemapping(SourceSkeleton, TargetSkeleton);
+	const bool bIsSkeletonRemappingValid = SkeletonRemapping.IsValid();
 
 	FGetBonePoseScratchArea& ScratchArea = FGetBonePoseScratchArea::Get();
 	BoneTrackArray& RotationScalePairs = ScratchArea.RotationScalePairs;
@@ -90,7 +91,7 @@ void DecompressPose(FCompactPose& OutPose,
 		for (int32 TrackIndex = (bFirstTrackIsRootBone ? 1 : 0); TrackIndex < NumTracks; TrackIndex++)
 		{
 			const int32 SourceSkeletonBoneIndex = CompressedData.GetSkeletonIndexFromTrackIndex(TrackIndex);
-			const int32 TargetSkeletonBoneIndex = SkeletonRemapping.IsValid() ? SkeletonRemapping.GetTargetSkeletonBoneIndex(SourceSkeletonBoneIndex) : SourceSkeletonBoneIndex;
+			const int32 TargetSkeletonBoneIndex = bIsSkeletonRemappingValid ? SkeletonRemapping.GetTargetSkeletonBoneIndex(SourceSkeletonBoneIndex) : SourceSkeletonBoneIndex;
 
 			if (TargetSkeletonBoneIndex != INDEX_NONE)
 			{
