AnimNext: Fix non-unity build

diff --git a/Engine/Plugins/Experimental/AnimNext/Source/AnimNext/Private/Graph/RigUnit_AnimNextWriteSkeletalMeshComponentPose.cpp b/Engine/Plugins/Experimental/AnimNext/Source/AnimNext/Private/Graph/RigUnit_AnimNextWriteSkeletalMeshComponentPose.cpp
--- a/Engine/Plugins/Experimental/AnimNext/Source/AnimNext/Private/Graph/RigUnit_AnimNextWriteSkeletalMeshComponentPose.cpp
+++ b/Engine/Plugins/Experimental/AnimNext/Source/AnimNext/Private/Graph/RigUnit_AnimNextWriteSkeletalMeshComponentPose.cpp
@@ -8,6 +8,7 @@
 #include "Components/SkeletalMeshComponent.h"
 #include "Engine/SkeletalMesh.h"
 #include "Logging/StructuredLog.h"
+#include "Module/AnimNextModuleInstance.h"
 #include "Module/AnimNextSkeletalMeshComponentReferenceComponent.h"
 
 #include UE_INLINE_GENERATED_CPP_BY_NAME(RigUnit_AnimNextWriteSkeletalMeshComponentPose)
diff --git a/Engine/Plugins/Experimental/AnimNext/Source/AnimNext/Private/Graph/RigUnit_GetActorTransform.cpp b/Engine/Plugins/Experimental/AnimNext/Source/AnimNext/Private/Graph/RigUnit_GetActorTransform.cpp
--- a/Engine/Plugins/Experimental/AnimNext/Source/AnimNext/Private/Graph/RigUnit_GetActorTransform.cpp
+++ b/Engine/Plugins/Experimental/AnimNext/Source/AnimNext/Private/Graph/RigUnit_GetActorTransform.cpp
@@ -3,6 +3,7 @@
 #include "RigUnit_GetActorTransform.h"
 
 #include "Component/AnimNextComponent.h"
+#include "Module/AnimNextModuleInstance.h"
 
 void FAnimNextActorTransformComponent::OnInitialize()
 {
diff --git a/Engine/Plugins/Experimental/AnimNext/Source/AnimNext/Private/Graph/RigUnit_MakeReferencePoseFromSkeletalMeshComponent.cpp b/Engine/Plugins/Experimental/AnimNext/Source/AnimNext/Private/Graph/RigUnit_MakeReferencePoseFromSkeletalMeshComponent.cpp
--- a/Engine/Plugins/Experimental/AnimNext/Source/AnimNext/Private/Graph/RigUnit_MakeReferencePoseFromSkeletalMeshComponent.cpp
+++ b/Engine/Plugins/Experimental/AnimNext/Source/AnimNext/Private/Graph/RigUnit_MakeReferencePoseFromSkeletalMeshComponent.cpp
@@ -6,6 +6,7 @@
 #include "GenerationTools.h"
 #include "Components/SkeletalMeshComponent.h"
 #include "Engine/SkeletalMesh.h"
+#include "Module/AnimNextModuleInstance.h"
 #include "Module/AnimNextSkeletalMeshComponentReferenceComponent.h"
 
 #include UE_INLINE_GENERATED_CPP_BY_NAME(RigUnit_MakeReferencePoseFromSkeletalMeshComponent)
