AnimNext: Fix public -> internal header dependency

diff --git a/Engine/Plugins/Experimental/AnimNext/Source/AnimNext/Private/Module/RigUnit_AnimNextModuleEventDependencies.cpp b/Engine/Plugins/Experimental/AnimNext/Source/AnimNext/Private/Module/RigUnit_AnimNextModuleEventDependencies.cpp
--- a/Engine/Plugins/Experimental/AnimNext/Source/AnimNext/Private/Module/RigUnit_AnimNextModuleEventDependencies.cpp
+++ b/Engine/Plugins/Experimental/AnimNext/Source/AnimNext/Private/Module/RigUnit_AnimNextModuleEventDependencies.cpp
@@ -223,6 +223,12 @@ TArray<FRigVMUserWorkflow> FRigUnit_AnimNextModuleDependenciesBase::GetSupported
 	return Workflows;
 }
 
+FRigVMTrait_ModuleEventDependency::FRigVMTrait_ModuleEventDependency()
+	: Ordering(EAnimNextModuleEventDependencyOrdering::Before)
+	, EventName(FRigUnit_AnimNextPrePhysicsEvent::DefaultEventName)
+{
+}
+
 FRigUnit_AnimNextModuleAddDependencies_Execute()
 {
 	UE::AnimNext::ApplyDependenciesHelper(ExecuteContext, [](const UE::AnimNext::FModuleDependencyContext& InContext, const FRigVMTrait_ModuleEventDependency& InDependency)
diff --git a/Engine/Plugins/Experimental/AnimNext/Source/AnimNext/Public/Module/RigVMTrait_ModuleEventDependency.h b/Engine/Plugins/Experimental/AnimNext/Source/AnimNext/Public/Module/RigVMTrait_ModuleEventDependency.h
--- a/Engine/Plugins/Experimental/AnimNext/Source/AnimNext/Public/Module/RigVMTrait_ModuleEventDependency.h
+++ b/Engine/Plugins/Experimental/AnimNext/Source/AnimNext/Public/Module/RigVMTrait_ModuleEventDependency.h
@@ -3,7 +3,6 @@
 #pragma once
 
 #include "CoreMinimal.h"
-#include "Module/RigUnit_AnimNextModuleEvents.h"
 #include "RigVMCore/RigVMTrait.h"
 #include "RigVMTrait_ModuleEventDependency.generated.h"
 
@@ -44,6 +43,8 @@ struct ANIMNEXT_API FRigVMTrait_ModuleEventDependency : public FRigVMTrait
 {
 	GENERATED_BODY()
 
+	FRigVMTrait_ModuleEventDependency();
+
 	// Override point that adds the dependency
 	virtual void OnAddDependency(const UE::AnimNext::FModuleDependencyContext& InContext) const {}
 
@@ -52,9 +53,9 @@ struct ANIMNEXT_API FRigVMTrait_ModuleEventDependency : public FRigVMTrait
 
 	// How to execute relative to the event
 	UPROPERTY(EditAnywhere, Category="Dependency")
-	EAnimNextModuleEventDependencyOrdering Ordering = EAnimNextModuleEventDependencyOrdering::Before;
+	EAnimNextModuleEventDependencyOrdering Ordering;
 
 	// The event to execute relative to
 	UPROPERTY(EditAnywhere, Category="Dependency", meta=(CustomWidget=AnimNextModuleEvent))
-	FName EventName = FRigUnit_AnimNextPrePhysicsEvent::DefaultEventName;
+	FName EventName;
 };
