Fix raw curve sampling in AnimNext

diff --git a/Engine/Plugins/Experimental/AnimNext/Source/AnimNext/Private/DecompressionTools.cpp b/Engine/Plugins/Experimental/AnimNext/Source/AnimNext/Private/DecompressionTools.cpp
--- a/Engine/Plugins/Experimental/AnimNext/Source/AnimNext/Private/DecompressionTools.cpp
+++ b/Engine/Plugins/Experimental/AnimNext/Source/AnimNext/Private/DecompressionTools.cpp
@@ -249,18 +249,7 @@ void FDecompressionTools::GetBonePose_AdditiveMeshRotationOnly(const UAnimSequen
 
 void FDecompressionTools::GetAnimationCurves(const UAnimSequence* AnimSequence, const FAnimExtractContext& ExtractionContext, FBlendedCurve& OutCurves, bool bForceUseRawData)
 {
-	check(!bForceUseRawData || CanEvaluateRawAnimationData(AnimSequence));
-
-	if (bForceUseRawData)
-	{
-		//FAnimationPoseData RemappedPoseData(OutAnimationPoseData.GetPose(), OutCurves, OutAnimationPoseData.GetAttributes());
-		//const UE::Anim::DataModel::FEvaluationContext EvaluationContext(ExtractionContext.CurrentTime, DataModelInterface->GetFrameRate(), GetRetargetTransformsSourceName(), GetRetargetTransforms(), Interpolation);
-		//DataModelInterface->Evaluate(RemappedPoseData, EvaluationContext);
-	}
-	else
-	{
-		AnimSequence->EvaluateCurveData(OutCurves, ExtractionContext.CurrentTime, bForceUseRawData);
-	}
+	AnimSequence->EvaluateCurveData(OutCurves, ExtractionContext.CurrentTime, bForceUseRawData);
 }
 
 void FDecompressionTools::GetAnimationAttributes(const UAnimSequence* AnimSequence, const FAnimExtractContext& ExtractionContext, const FReferencePose& RefPose, UE::Anim::FStackAttributeContainer& OutAttributes, bool bForceUseRawData)
