AnimNext: Exposed component-level tick dependencies to Blueprint

Allows meshes etc. to be set as prerequisites/subsequents of animation work

diff --git a/Engine/Plugins/Experimental/AnimNext/Source/AnimNext/Private/Component/AnimNextComponent.cpp b/Engine/Plugins/Experimental/AnimNext/Source/AnimNext/Private/Component/AnimNextComponent.cpp
--- a/Engine/Plugins/Experimental/AnimNext/Source/AnimNext/Private/Component/AnimNextComponent.cpp
+++ b/Engine/Plugins/Experimental/AnimNext/Source/AnimNext/Private/Component/AnimNextComponent.cpp
@@ -938,6 +938,11 @@ void UAnimNextComponent::AddPrerequisite(UObject* InObject, FTickFunction& InTic
 	}
 }
 
+void UAnimNextComponent::AddComponentPrerequisite(UActorComponent* InComponent, FName InEventName)
+{
+	AddPrerequisite(InComponent, InComponent->PrimaryComponentTick, InEventName);
+}
+
 void UAnimNextComponent::AddSubsequent(UObject* InObject, FTickFunction& InTickFunction, FName InEventName)
 {
 	if(Subsystem)
@@ -946,6 +951,11 @@ void UAnimNextComponent::AddSubsequent(UObject* InObject, FTickFunction& InTickF
 	}
 }
 
+void UAnimNextComponent::AddComponentSubsequent(UActorComponent* InComponent, FName InEventName)
+{
+	AddSubsequent(InComponent, InComponent->PrimaryComponentTick, InEventName);
+}
+
 void UAnimNextComponent::RemovePrerequisite(UObject* InObject, FTickFunction& InTickFunction, FName InEventName)
 {
 	if(Subsystem)
@@ -954,6 +964,11 @@ void UAnimNextComponent::RemovePrerequisite(UObject* InObject, FTickFunction& In
 	}
 }
 
+void UAnimNextComponent::RemoveComponentPrerequisite(UActorComponent* InComponent, FName InEventName)
+{
+	RemovePrerequisite(InComponent, InComponent->PrimaryComponentTick, InEventName);
+}
+
 void UAnimNextComponent::RemoveSubsequent(UObject* InObject, FTickFunction& InTickFunction, FName InEventName)
 {
 	if(Subsystem)
@@ -962,6 +977,11 @@ void UAnimNextComponent::RemoveSubsequent(UObject* InObject, FTickFunction& InTi
 	}
 }
 
+void UAnimNextComponent::RemoveComponentSubsequent(UActorComponent* InComponent, FName InEventName)
+{
+	RemoveSubsequent(InComponent, InComponent->PrimaryComponentTick, InEventName);
+}
+
 void UAnimNextComponent::AddModuleEventPrerequisite(FName InEventName, UAnimNextComponent* OtherAnimNextComponent, FName OtherEventName)
 {
 	if (!OtherAnimNextComponent)
diff --git a/Engine/Plugins/Experimental/AnimNext/Source/AnimNext/Public/Component/AnimNextComponent.h b/Engine/Plugins/Experimental/AnimNext/Source/AnimNext/Public/Component/AnimNextComponent.h
--- a/Engine/Plugins/Experimental/AnimNext/Source/AnimNext/Public/Component/AnimNextComponent.h
+++ b/Engine/Plugins/Experimental/AnimNext/Source/AnimNext/Public/Component/AnimNextComponent.h
@@ -215,24 +215,50 @@ public:
 	// @param	InEventName			The event to add the dependency to
 	ANIMNEXT_API void AddPrerequisite(UObject* InObject, FTickFunction& InTickFunction, FName InEventName);
 
+	// Add a prerequisite dependency on the component's primary tick function to the specified event
+	// The component will tick before the event
+	// @param	Component			The component to add as a prerequisite 
+	// @param	EventName			The event to add the dependency to
+	UFUNCTION(BlueprintCallable, Category = "AnimNext")
+	ANIMNEXT_API void AddComponentPrerequisite(UActorComponent* Component, FName EventName);
+
 	// Add a subsequent tick function dependency to the specified event
 	// @param	InObject			The object that owns the tick function
 	// @param	InTickFunction		The tick function to depend on
 	// @param	InEventName			The event to add the dependency to
 	ANIMNEXT_API void AddSubsequent(UObject* InObject, FTickFunction& InTickFunction, FName InEventName);
 
+	// Add a subsequent dependency on the component's primary tick function to the specified event
+	// The component will tick after the event
+	// @param	Component			The component to add as a subsequent of the event
+	// @param	EventName			The event to add the dependency to
+	UFUNCTION(BlueprintCallable, Category = "AnimNext")
+	ANIMNEXT_API void AddComponentSubsequent(UActorComponent* Component, FName EventName);
+
 	// Remove a prerequisite tick function dependency from the specified event
 	// @param	InObject			The object that owns the tick function
 	// @param	InTickFunction		The tick function that was depended on
 	// @param	InEventName			The event to remove the dependency from
 	ANIMNEXT_API void RemovePrerequisite(UObject* InObject, FTickFunction& InTickFunction, FName InEventName);
 
+	// Remove a prerequisite on the component's primary tick function from the specified event
+	// @param	Component			The component to remove as a prerequisite 
+	// @param	EventName			The event to add the dependency to
+	UFUNCTION(BlueprintCallable, Category = "AnimNext")
+	ANIMNEXT_API void RemoveComponentPrerequisite(UActorComponent* Component, FName EventName);
+	
 	// Remove a prerequisite tick function dependency from the specified event
 	// @param	InObject			The object that owns the tick function
 	// @param	InTickFunction		The tick function that was depended on
 	// @param	InEventName			The event to remove the dependency from
 	ANIMNEXT_API void RemoveSubsequent(UObject* InObject, FTickFunction& InTickFunction, FName InEventName);
 
+	// Remove a subsequent dependency on the component's primary tick function from the specified event
+	// @param	Component			The component to remove as a subsequent of the event
+	// @param	EventName			The event to add the dependency to
+	UFUNCTION(BlueprintCallable, Category = "AnimNext")
+	ANIMNEXT_API void RemoveComponentSubsequent(UActorComponent* Component, FName EventName);
+	
 	// Add a prerequisite anim next event dependency to the specified event
 	// @param	InEventName			The event name in this component
 	// @param	InTickFunction		The other component we want a prerequisite on
