Fix modules not removing themselves as dependencies inside other modules when unregistering handle.
* Also fix logic error when handling remove prerequisite inside FlushPendingActions

diff --git a/Engine/Plugins/Experimental/AnimNext/Source/AnimNext/Internal/Module/AnimNextModuleInstance.h b/Engine/Plugins/Experimental/AnimNext/Source/AnimNext/Internal/Module/AnimNextModuleInstance.h
--- a/Engine/Plugins/Experimental/AnimNext/Source/AnimNext/Internal/Module/AnimNextModuleInstance.h
+++ b/Engine/Plugins/Experimental/AnimNext/Source/AnimNext/Internal/Module/AnimNextModuleInstance.h
@@ -225,6 +225,9 @@ private:
 	// All pre-requisite references, only modified from the game thread
 	TArray<FPrerequisiteReference> PrerequisiteRefs;
 
+	// All subsequent references, only modified from the game thread
+	TArray<UE::AnimNext::FModuleHandle> SubsequentRefs;
+
 	// Input event list to be processed on the next update
 	UE::AnimNext::FTraitEventList InputEventList;
 
diff --git a/Engine/Plugins/Experimental/AnimNext/Source/AnimNext/Private/Component/AnimNextWorldSubsystem.cpp b/Engine/Plugins/Experimental/AnimNext/Source/AnimNext/Private/Component/AnimNextWorldSubsystem.cpp
--- a/Engine/Plugins/Experimental/AnimNext/Source/AnimNext/Private/Component/AnimNextWorldSubsystem.cpp
+++ b/Engine/Plugins/Experimental/AnimNext/Source/AnimNext/Private/Component/AnimNextWorldSubsystem.cpp
@@ -161,7 +161,7 @@ void UAnimNextWorldSubsystem::FlushPendingActions()
 					FAnimNextModuleInstance* PrerequisiteInstance = Instances.TryGet(PendingAction.Payload.Get<UE::AnimNext::FModuleHandle>());
 					if(PrerequisiteInstance != nullptr && SubsequentInstance != nullptr)
 					{
-						if (!SubsequentInstance->IsPrerequisite(*PrerequisiteInstance))
+						if (SubsequentInstance->IsPrerequisite(*PrerequisiteInstance))
 						{							
 							SubsequentInstance->RemovePrerequisite(*PrerequisiteInstance);
 						}
@@ -217,6 +217,12 @@ void UAnimNextWorldSubsystem::UnregisterHandle(UE::AnimNext::FModuleHandle& InOu
 		// Remove all tick dependencies immediately, as once the handle has been invalidated there is no way for external systems to remove their dependencies
 		Instance.RemoveAllTickDependencies();
 
+		// Remove this module as a dependency on all other modules that depend upon it before the handle is invalidated
+		for (const UE::AnimNext::FModuleHandle& SubsequentHandle : Instance.SubsequentRefs)
+		{
+			RemoveModuleDependencyHandle(InOutHandle, SubsequentHandle);
+		}
+
 		PendingActions.Emplace(InOutHandle, FModulePendingAction::EType::ReleaseHandle);
 		InOutHandle.Reset();
 	}
diff --git a/Engine/Plugins/Experimental/AnimNext/Source/AnimNext/Private/Module/AnimNextModuleInstance.cpp b/Engine/Plugins/Experimental/AnimNext/Source/AnimNext/Private/Module/AnimNextModuleInstance.cpp
--- a/Engine/Plugins/Experimental/AnimNext/Source/AnimNext/Private/Module/AnimNextModuleInstance.cpp
+++ b/Engine/Plugins/Experimental/AnimNext/Source/AnimNext/Private/Module/AnimNextModuleInstance.cpp
@@ -687,6 +687,7 @@ void FAnimNextModuleInstance::AddPrerequisite(FAnimNextModuleInstance& InPrerequ
 		}
 
 		PrerequisiteRefs.Add({ InPrerequisiteInstance.Handle, 1});
+		InPrerequisiteInstance.SubsequentRefs.Add(GetHandle());
 	}
 }
 
@@ -717,6 +718,9 @@ void FAnimNextModuleInstance::RemovePrerequisite(FAnimNextModuleInstance& InPrer
 			}
 
 			PrerequisiteRefs.RemoveAtSwap(FoundIndex, EAllowShrinking::No);
+
+			const int32 CountRemoved = InPrerequisiteInstance.SubsequentRefs.RemoveSwap(GetHandle(), EAllowShrinking::No);
+			check(CountRemoved == 1);
 		}
 	}
 }
