Fix static analysis in PVS Studio (blind)

diff --git a/Engine/Plugins/Experimental/AnimNext/Source/AnimNext/Public/LODPose.h b/Engine/Plugins/Experimental/AnimNext/Source/AnimNext/Public/LODPose.h
--- a/Engine/Plugins/Experimental/AnimNext/Source/AnimNext/Public/LODPose.h
+++ b/Engine/Plugins/Experimental/AnimNext/Source/AnimNext/Public/LODPose.h
@@ -232,7 +232,7 @@ struct FLODPose
 		const TArrayView<const FBoneIndexType> LODBoneIndexToParentMap = RefPose->GetLODBoneIndexToParentLODBoneIndexMap(LODLevel); 
 		
 		FBoneIndexType ParentOfChild = LODBoneIndexToParentMap[ChildLODBoneIndex];
-		while (ParentOfChild != INDEX_NONE)
+		while (ParentOfChild != (FBoneIndexType)INDEX_NONE)	// Note: Explicit cast to make static analysis happy
 		{
 			if (ParentOfChild == ParentLODBoneIndex)
 			{
