Prevent possible crashes when there is a potential mismatch in bones between retarget transforms vs the animation transforms

diff --git a/Engine/Plugins/Experimental/AnimNext/Source/AnimNext/Private/DecompressionTools.cpp b/Engine/Plugins/Experimental/AnimNext/Source/AnimNext/Private/DecompressionTools.cpp
--- a/Engine/Plugins/Experimental/AnimNext/Source/AnimNext/Private/DecompressionTools.cpp
+++ b/Engine/Plugins/Experimental/AnimNext/Source/AnimNext/Private/DecompressionTools.cpp
@@ -564,6 +564,10 @@ void FDecompressionTools::DecompressPose(FLODPose& OutAnimationPoseData,
 		{
 			const int32 LODBoneIndex(BonePair.AtomIndex);
 			const int32 SourceSkeletonBoneIndex = BonePair.TrackIndex;
+			if (!AuthoredOnRefSkeleton.IsValidIndex(SourceSkeletonBoneIndex))
+			{
+				continue;
+			}
 
 			// @todo - precache that in FBoneContainer when we have SkeletonIndex->TrackIndex mapping. So we can just apply scale right away.
 			float const SourceTranslationLength = AuthoredOnRefSkeleton[SourceSkeletonBoneIndex].GetTranslation().Size();
@@ -585,6 +589,10 @@ void FDecompressionTools::DecompressPose(FLODPose& OutAnimationPoseData,
 		{
 			const int32 LODBoneIndex(BonePair.AtomIndex);
 			const int32 SourceSkeletonBoneIndex = BonePair.TrackIndex;
+			if (!AuthoredOnRefSkeleton.IsValidIndex(SourceSkeletonBoneIndex))
+			{
+				continue;
+			}
 
 			const FTransform& RefPoseTransform = ReferencePose.GetRefPoseTransform(LODBoneIndex);
 
@@ -711,12 +719,15 @@ void FDecompressionTools::RetargetBoneTransform(const FReferencePose& ReferenceP
 		case EBoneTranslationRetargetingMode::AnimationScaled:
 		{
 			// @todo - precache that in FBoneContainer when we have SkeletonIndex->TrackIndex mapping. So we can just apply scale right away.
-			const TArray<FTransform>& SkeletonRefPoseArray = RetargetTransforms;
-			const float SourceTranslationLength = SkeletonRefPoseArray[SourceSkeletonBoneIndex].GetTranslation().Size();
-			if (SourceTranslationLength > UE_KINDA_SMALL_NUMBER)
+			if (RetargetTransforms.IsValidIndex(SourceSkeletonBoneIndex))
 			{
-				const float TargetTranslationLength = ReferencePose.GetRefPoseTranslation(LODBoneIndex).Size();
-				BoneTransform.ScaleTranslation(TargetTranslationLength / SourceTranslationLength);
+				const TArray<FTransform>& SkeletonRefPoseArray = RetargetTransforms;
+				const float SourceTranslationLength = SkeletonRefPoseArray[SourceSkeletonBoneIndex].GetTranslation().Size();
+				if (SourceTranslationLength > UE_KINDA_SMALL_NUMBER)
+				{
+					const float TargetTranslationLength = ReferencePose.GetRefPoseTranslation(LODBoneIndex).Size();
+					BoneTransform.ScaleTranslation(TargetTranslationLength / SourceTranslationLength);
+				}
 			}
 			break;
 		}
@@ -737,17 +748,20 @@ void FDecompressionTools::RetargetBoneTransform(const FReferencePose& ReferenceP
 				const FTransform& RefPoseTransform = ReferencePose.GetRefPoseTransform(LODBoneIndex);
 
 				// Remap the base pose onto the target skeleton so that we are working entirely in target space
-				FTransform BaseTransform = AuthoredOnRefSkeleton[SourceSkeletonBoneIndex];
-				if (SkeletonRemapping.RequiresReferencePoseRetarget())
+				if (AuthoredOnRefSkeleton.IsValidIndex(SourceSkeletonBoneIndex))
 				{
-					BaseTransform = SkeletonRemapping.RetargetBoneTransformToTargetSkeleton(TargetSkeletonBoneIndex, BaseTransform);
-				}
+					FTransform BaseTransform = AuthoredOnRefSkeleton[SourceSkeletonBoneIndex];
+					if (SkeletonRemapping.RequiresReferencePoseRetarget())
+					{
+						BaseTransform = SkeletonRemapping.RetargetBoneTransformToTargetSkeleton(TargetSkeletonBoneIndex, BaseTransform);
+					}
 
-				// Apply the retargeting as if it were an additive difference between the current skeleton and the retarget skeleton. 
-				BoneTransform.SetRotation(BoneTransform.GetRotation() * BaseTransform.GetRotation().Inverse() * RefPoseTransform.GetRotation());
-				BoneTransform.SetTranslation(BoneTransform.GetTranslation() + (RefPoseTransform.GetTranslation() - BaseTransform.GetTranslation()));
-				BoneTransform.SetScale3D(BoneTransform.GetScale3D() * (RefPoseTransform.GetScale3D() * BaseTransform.GetSafeScaleReciprocal(BaseTransform.GetScale3D())));
-				BoneTransform.NormalizeRotation();
+					// Apply the retargeting as if it were an additive difference between the current skeleton and the retarget skeleton. 
+					BoneTransform.SetRotation(BoneTransform.GetRotation() * BaseTransform.GetRotation().Inverse() * RefPoseTransform.GetRotation());
+					BoneTransform.SetTranslation(BoneTransform.GetTranslation() + (RefPoseTransform.GetTranslation() - BaseTransform.GetTranslation()));
+					BoneTransform.SetScale3D(BoneTransform.GetScale3D() * (RefPoseTransform.GetScale3D() * BaseTransform.GetSafeScaleReciprocal(BaseTransform.GetScale3D())));
+					BoneTransform.NormalizeRotation();
+				}
 			}
 			break;
 		}
diff --git a/Engine/Source/Runtime/Engine/Private/Animation/AnimationDecompression.cpp b/Engine/Source/Runtime/Engine/Private/Animation/AnimationDecompression.cpp
--- a/Engine/Source/Runtime/Engine/Private/Animation/AnimationDecompression.cpp
+++ b/Engine/Source/Runtime/Engine/Private/Animation/AnimationDecompression.cpp
@@ -285,6 +285,10 @@ void DecompressPose(FCompactPose& OutPose,
 		{
 			const FCompactPoseBoneIndex BoneIndex(BonePair.AtomIndex); //Nasty, we break our type safety, code in the lower levels should be adjusted for this
 			const int32 SourceSkeletonBoneIndex = BonePair.TrackIndex;
+			if (!AuthoredOnRefSkeleton.IsValidIndex(SourceSkeletonBoneIndex))
+			{
+				continue;
+			}
 
 			// @todo - precache that in FBoneContainer when we have SkeletonIndex->TrackIndex mapping. So we can just apply scale right away.
 			float const SourceTranslationLength = AuthoredOnRefSkeleton[SourceSkeletonBoneIndex].GetTranslation().Size();
@@ -306,6 +310,10 @@ void DecompressPose(FCompactPose& OutPose,
 		{
 			const FCompactPoseBoneIndex BoneIndex(BonePair.AtomIndex); //Nasty, we break our type safety, code in the lower levels should be adjusted for this
 			const int32 SourceSkeletonBoneIndex = BonePair.TrackIndex;
+			if (!AuthoredOnRefSkeleton.IsValidIndex(SourceSkeletonBoneIndex))
+			{
+				continue;
+			}
 
 			const FTransform& RefPoseTransform = RequiredBones.GetRefPoseTransform(BoneIndex);
 
diff --git a/Engine/Source/Runtime/Engine/Private/Animation/AnimationRuntime.cpp b/Engine/Source/Runtime/Engine/Private/Animation/AnimationRuntime.cpp
--- a/Engine/Source/Runtime/Engine/Private/Animation/AnimationRuntime.cpp
+++ b/Engine/Source/Runtime/Engine/Private/Animation/AnimationRuntime.cpp
@@ -2931,12 +2931,15 @@ void FAnimationRuntime::RetargetBoneTransform(const USkeleton* SourceSkeleton, c
 			case EBoneTranslationRetargetingMode::AnimationScaled:
 			{
 				// @todo - precache that in FBoneContainer when we have SkeletonIndex->TrackIndex mapping. So we can just apply scale right away.
-				const TArray<FTransform>& SkeletonRefPoseArray = RetargetTransforms;
-				const float SourceTranslationLength = SkeletonRefPoseArray[SourceSkeletonBoneIndex].GetTranslation().Size();
-				if (SourceTranslationLength > UE_KINDA_SMALL_NUMBER)
+				if (RetargetTransforms.IsValidIndex(SourceSkeletonBoneIndex))
 				{
-					const float TargetTranslationLength = RequiredBones.GetRefPoseTransform(BoneIndex).GetTranslation().Size();
-					BoneTransform.ScaleTranslation(TargetTranslationLength / SourceTranslationLength);
+					const TArray<FTransform>& SkeletonRefPoseArray = RetargetTransforms;
+					const float SourceTranslationLength = SkeletonRefPoseArray[SourceSkeletonBoneIndex].GetTranslation().Size();
+					if (SourceTranslationLength > UE_KINDA_SMALL_NUMBER)
+					{
+						const float TargetTranslationLength = RequiredBones.GetRefPoseTransform(BoneIndex).GetTranslation().Size();
+						BoneTransform.ScaleTranslation(TargetTranslationLength / SourceTranslationLength);
+					}
 				}
 				break;
 			}
