AnimNext fixes to allow running of post-mover events

diff --git a/Engine/Plugins/Experimental/AnimNext/Source/AnimNext/Private/Module/AnimNextModuleInstance.cpp b/Engine/Plugins/Experimental/AnimNext/Source/AnimNext/Private/Module/AnimNextModuleInstance.cpp
--- a/Engine/Plugins/Experimental/AnimNext/Source/AnimNext/Private/Module/AnimNextModuleInstance.cpp
+++ b/Engine/Plugins/Experimental/AnimNext/Source/AnimNext/Private/Module/AnimNextModuleInstance.cpp
@@ -113,6 +113,31 @@ static TConstArrayView<FImplementedModuleEvent> GetImplementedModuleEvents(const
 			NewEvent.bUserEvent = Event.IsUserEvent();
 			NewEvent.bIsTask = Event.IsTask();
 			NewEvent.bIsGameThreadTask = Event.IsGameThreadTask();
+
+			// User events can override their event name etc. via parameters 
+			if (Function->Struct->IsChildOf(FRigUnit_AnimNextUserEvent::StaticStruct()))
+			{
+				// Pull the values out of the literal memory
+				FRigVMOperandArray Operands = ByteCode.GetOperandsForExecuteOp(Instruction);
+				check(Function->ArgumentNames.Num() == Operands.Num());
+				int32 NumOperands = Operands.Num();
+				URigVM* NonConstVM = const_cast<URigVM*>(VM);
+				for (int32 OperandIndex = 0; OperandIndex < NumOperands; ++OperandIndex)
+				{
+					const FRigVMOperand& Operand = Operands[OperandIndex];
+					FName OperandName = Function->ArgumentNames[OperandIndex];
+					if (OperandName == GET_MEMBER_NAME_CHECKED(FRigUnit_AnimNextUserEvent, Name))
+					{
+						check(Operand.GetMemoryType() == ERigVMMemoryType::Literal);
+						NewEvent.EventName = *NonConstVM->LiteralMemoryStorage.GetData<FName>(Operand.GetRegisterIndex());
+					}
+					else if (OperandName == GET_MEMBER_NAME_CHECKED(FRigUnit_AnimNextUserEvent, SortOrder))
+					{
+						check(Operand.GetMemoryType() == ERigVMMemoryType::Literal);
+						NewEvent.SortOrder = *NonConstVM->LiteralMemoryStorage.GetData<int32>(Operand.GetRegisterIndex());
+					}
+				}
+			}
 		}
 	}
 
