Optimize pose remapping

diff --git a/Engine/Plugins/Experimental/AnimNext/Source/AnimNext/Private/GenerationTools.cpp b/Engine/Plugins/Experimental/AnimNext/Source/AnimNext/Private/GenerationTools.cpp
--- a/Engine/Plugins/Experimental/AnimNext/Source/AnimNext/Private/GenerationTools.cpp
+++ b/Engine/Plugins/Experimental/AnimNext/Source/AnimNext/Private/GenerationTools.cpp
@@ -617,10 +617,36 @@ void FGenerationTools::FixLODRequiredBones(const int32 NumLODs
 
 	check(SourcePose.GetNumBones() == NumLODBones);
 
+#if DEFAULT_SOA && DEFAULT_SOA_VIEW
+	const double* RESTRICT RotationPtr = reinterpret_cast<const double*>(SourcePose.LocalTransformsView.Rotations.GetData());
+	const double* RESTRICT RotationEndPtr = RotationPtr + NumLODBones * 4;
+
+	// Our SoA buffer is contiguous
+	// Because the translations/scales have the same size (FVector), each entry is a fixed offset apart and we can use a single ptr/offset pair
+	const double* RESTRICT TranslationPtr = reinterpret_cast<const double*>(SourcePose.LocalTransformsView.Translations.GetData());
+	const int64 ScaleOffset = reinterpret_cast<const double*>(SourcePose.LocalTransformsView.Scales3D.GetData()) - TranslationPtr;
+
+	const FBoneIndexType* RESTRICT LODBoneIndexPtr = LODBoneIndexes.GetData();
+
+	while (RotationPtr < RotationEndPtr)
+	{
+		VectorRegister4Double Rotation = VectorLoadAligned(RotationPtr);
+		VectorRegister4Double Translation = VectorSet_W0(VectorLoadAligned(TranslationPtr));
+		VectorRegister4Double Scale = VectorSet_W0(VectorLoadAligned(TranslationPtr + ScaleOffset));
+
+		const int32 DestTransformIndex = *LODBoneIndexPtr;
+		TargetTransforms[DestTransformIndex] = FTransform(Rotation, Translation, Scale);
+
+		RotationPtr += 4;
+		TranslationPtr += 3;
+		LODBoneIndexPtr++;
+	}
+#else
 	for (int32 LODBoneIndex = 0; LODBoneIndex < NumLODBones; ++LODBoneIndex)
 	{
 		TargetTransforms[LODBoneIndexes[LODBoneIndex]] = SourcePose.LocalTransformsView[LODBoneIndex];
 	}
+#endif
 }
 
 void FGenerationTools::RemapAttributes(
diff --git a/Engine/Plugins/Experimental/AnimNext/Source/AnimNext/Public/TransformArray.h b/Engine/Plugins/Experimental/AnimNext/Source/AnimNext/Public/TransformArray.h
--- a/Engine/Plugins/Experimental/AnimNext/Source/AnimNext/Public/TransformArray.h
+++ b/Engine/Plugins/Experimental/AnimNext/Source/AnimNext/Public/TransformArray.h
@@ -269,6 +269,11 @@ using FTransformArrayAoSStack = TTransformArrayAoS<FAnimStackAllocator>;
 template <typename AllocatorType>
 struct TTransformArraySoA
 {
+	// We add a bit of padding at the end of our buffer to allow 32-byte wide load/stores
+	// without the risk of triggering access exceptions
+	static constexpr int32 PaddingSize = sizeof(double);
+	static constexpr int32 TransformSize = sizeof(FVector) + sizeof(FQuat) + sizeof(FVector);
+
 	TArrayView<FVector> Translations;
 	TArrayView<FQuat> Rotations;
 	TArrayView<FVector> Scales3D;
@@ -325,41 +330,32 @@ struct TTransformArraySoA
 
 	inline void Empty(int32 Slack = 0)
 	{
-		AllocatedMemory.Empty(Slack);
+		const int32 Padding = Slack != 0 ? PaddingSize : 0;
+		AllocatedMemory.Empty(Slack + Padding);
 		UpdateViews(AllocatedMemory.GetData(), 0);
 	}
 
 	inline void Reset(int32 NumTransforms)
 	{
-		constexpr int32 TransformSize = sizeof(FVector) + sizeof(FQuat) + sizeof(FVector);
-		AllocatedMemory.Reset(NumTransforms * TransformSize);
+		const int32 Padding = NumTransforms != 0 ? PaddingSize : 0;
+		AllocatedMemory.Reset(NumTransforms * TransformSize + Padding);
 	}
 
 	void SetNum(int32 NumTransforms, EAllowShrinking AllowShrinking = EAllowShrinking::Default)
 	{
-		constexpr int32 TransformSize = sizeof(FVector) + sizeof(FQuat) + sizeof(FVector);
-		AllocatedMemory.SetNum(NumTransforms * TransformSize, AllowShrinking);
+		const int32 Padding = NumTransforms != 0 ? PaddingSize : 0;
+		AllocatedMemory.SetNum(NumTransforms * TransformSize + Padding, AllowShrinking);
 
 		UpdateViews(AllocatedMemory.GetData(), NumTransforms);
 	}
-	UE_ALLOWSHRINKING_BOOL_DEPRECATED("SetNum")
-	FORCEINLINE void SetNum(int32 NumTransforms, bool bAllowShrinking)
-	{
-		SetNum(NumTransforms, bAllowShrinking ? EAllowShrinking::Yes : EAllowShrinking::No);
-	}
 
 	void SetNumUninitialized(int32 NumTransforms, EAllowShrinking AllowShrinking = EAllowShrinking::Default)
 	{
-		constexpr int32 TransformSize = sizeof(FVector) + sizeof(FQuat) + sizeof(FVector);
-		AllocatedMemory.SetNumUninitialized(NumTransforms * TransformSize, AllowShrinking);
+		const int32 Padding = NumTransforms != 0 ? PaddingSize : 0;
+		AllocatedMemory.SetNumUninitialized(NumTransforms * TransformSize + Padding, AllowShrinking);
 
 		UpdateViews(AllocatedMemory.GetData(), NumTransforms);
 	}
-	UE_ALLOWSHRINKING_BOOL_DEPRECATED("SetNumUninitialized")
-	FORCEINLINE void SetNumUninitialized(int32 NumTransforms, bool bAllowShrinking)
-	{
-		SetNumUninitialized(NumTransforms, bAllowShrinking ? EAllowShrinking::Yes : EAllowShrinking::No);
-	}
 
 	inline void SetIdentity(bool bAdditiveIdentity = false)
 	{
