Added anim next workspace exports caching, as behaviour differs during save (due to FindObject limitations while performing saving)

diff --git a/Engine/Plugins/Experimental/AnimNext/Source/AnimNextUncookedOnly/Internal/AnimNextAssetWorkspaceAssetUserData.h b/Engine/Plugins/Experimental/AnimNext/Source/AnimNextUncookedOnly/Internal/AnimNextAssetWorkspaceAssetUserData.h
--- a/Engine/Plugins/Experimental/AnimNext/Source/AnimNextUncookedOnly/Internal/AnimNextAssetWorkspaceAssetUserData.h
+++ b/Engine/Plugins/Experimental/AnimNext/Source/AnimNextUncookedOnly/Internal/AnimNextAssetWorkspaceAssetUserData.h
@@ -147,4 +147,6 @@ private:
 	GENERATED_BODY()
 
 	virtual void GetAssetRegistryTags(FAssetRegistryTagsContext Context) const override;
+
+	mutable FWorkspaceOutlinerItemExports CachedExports;
 };
diff --git a/Engine/Plugins/Experimental/AnimNext/Source/AnimNextUncookedOnly/Private/AnimNextAssetWorkspaceAssetUserData.cpp b/Engine/Plugins/Experimental/AnimNext/Source/AnimNextUncookedOnly/Private/AnimNextAssetWorkspaceAssetUserData.cpp
--- a/Engine/Plugins/Experimental/AnimNext/Source/AnimNextUncookedOnly/Private/AnimNextAssetWorkspaceAssetUserData.cpp
+++ b/Engine/Plugins/Experimental/AnimNext/Source/AnimNextUncookedOnly/Private/AnimNextAssetWorkspaceAssetUserData.cpp
@@ -12,34 +12,38 @@ void UAnimNextAssetWorkspaceAssetUserData::GetAssetRegistryTags(FAssetRegistryTa
 {
 	Super::GetAssetRegistryTags(Context);
 
-	FWorkspaceOutlinerItemExports Exports;
+	// Updated cached export data outside of saving call-path, and simply return cached data when actually saving 
+	if (!Context.IsSaving())
 	{
-		UAnimNextRigVMAsset* Asset = CastChecked<UAnimNextRigVMAsset>(GetOuter());
-		const UAnimNextRigVMAssetEditorData* GraphEditorData = UE::AnimNext::UncookedOnly::FUtils::GetEditorData<UAnimNextRigVMAssetEditorData>(Asset);
+		CachedExports.Exports.Reset();
 		{
-			FWorkspaceOutlinerItemExport& RootAssetExport = Exports.Exports.Add_GetRef(FWorkspaceOutlinerItemExport(Asset->GetFName(), Asset));
-
-			if(UAnimNextModule* Module = Cast<UAnimNextModule>(Asset))
-			{
-				RootAssetExport.GetData().InitializeAsScriptStruct(FAnimNextModuleOutlinerData::StaticStruct());
-			}
-			else if(UAnimNextDataInterface* DataInterface = Cast<UAnimNextDataInterface>(Asset))
+			UAnimNextRigVMAsset* Asset = CastChecked<UAnimNextRigVMAsset>(GetOuter());
+			const UAnimNextRigVMAssetEditorData* GraphEditorData = UE::AnimNext::UncookedOnly::FUtils::GetEditorData<UAnimNextRigVMAssetEditorData>(Asset);
 			{
-				RootAssetExport.GetData().InitializeAsScriptStruct(FAnimNextDataInterfaceOutlinerData::StaticStruct());
-			}
-			else
-			{
-				RootAssetExport.GetData().InitializeAsScriptStruct(FAnimNextRigVMAssetOutlinerData::StaticStruct());
+				FWorkspaceOutlinerItemExport& RootAssetExport = CachedExports.Exports.Add_GetRef(FWorkspaceOutlinerItemExport(Asset->GetFName(), Asset));
+
+				if(UAnimNextModule* Module = Cast<UAnimNextModule>(Asset))
+				{
+					RootAssetExport.GetData().InitializeAsScriptStruct(FAnimNextModuleOutlinerData::StaticStruct());
+				}
+				else if(UAnimNextDataInterface* DataInterface = Cast<UAnimNextDataInterface>(Asset))
+				{
+					RootAssetExport.GetData().InitializeAsScriptStruct(FAnimNextDataInterfaceOutlinerData::StaticStruct());
+				}
+				else
+				{
+					RootAssetExport.GetData().InitializeAsScriptStruct(FAnimNextRigVMAssetOutlinerData::StaticStruct());
+				}
+				FAnimNextRigVMAssetOutlinerData& Data = RootAssetExport.GetData().GetMutable<FAnimNextRigVMAssetOutlinerData>();
+				Data.SoftAssetPtr = Asset;
 			}
-			FAnimNextRigVMAssetOutlinerData& Data = RootAssetExport.GetData().GetMutable<FAnimNextRigVMAssetOutlinerData>();
-			Data.SoftAssetPtr = Asset;
+		
+			UE::AnimNext::UncookedOnly::FUtils::GetAssetOutlinerItems(GraphEditorData, CachedExports, Context);
 		}
-	
-		UE::AnimNext::UncookedOnly::FUtils::GetAssetOutlinerItems(GraphEditorData, Exports, Context);
 	}
 	
 	FString TagValue;
-	FWorkspaceOutlinerItemExports::StaticStruct()->ExportText(TagValue, &Exports, nullptr, nullptr, PPF_None, nullptr);
+	FWorkspaceOutlinerItemExports::StaticStruct()->ExportText(TagValue, &CachedExports, nullptr, nullptr, PPF_None, nullptr);
 	Context.AddTag(FAssetRegistryTag(UE::Workspace::ExportsWorkspaceItemsRegistryTag, TagValue, FAssetRegistryTag::TT_Hidden));
 }
 
