Optimize pose remapping

diff --git a/Engine/Plugins/Experimental/AnimNext/Source/AnimNext/Private/Graph/RemapPoseData.cpp b/Engine/Plugins/Experimental/AnimNext/Source/AnimNext/Private/Graph/RemapPoseData.cpp
--- a/Engine/Plugins/Experimental/AnimNext/Source/AnimNext/Private/Graph/RemapPoseData.cpp
+++ b/Engine/Plugins/Experimental/AnimNext/Source/AnimNext/Private/Graph/RemapPoseData.cpp
@@ -144,8 +144,43 @@ void FRemapPoseData::RemapPose(const UE::AnimNext::FLODPoseHeap& SourcePose,
 	// for memory reasons as we'd store duplicated data otherwise and bones are guaranteed to be sorted.
 	const int32 TargetLODIndex = FMath::Clamp(SourceLODLevel, 0, SourceToTargetBoneIndexMapPerLOD[SourceLODLevel].Num() - 1);
 	const TArray<FRemappedBone>& SourceToTargetBoneIndexMap = SourceToTargetBoneIndexMapPerLOD[SourceLODLevel][TargetLODIndex].BoneIndexMap;
+	const int32 NumToRemap = SourceToTargetBoneIndexMap.Num();
 
-	for (int32 MappingElementIndex = 0; MappingElementIndex < SourceToTargetBoneIndexMap.Num(); ++MappingElementIndex)
+#if DEFAULT_SOA && DEFAULT_SOA_VIEW
+	const double* RESTRICT SrcRotationPtr = reinterpret_cast<const double*>(SourcePose.LocalTransformsView.Rotations.GetData());
+	double* RESTRICT DstRotationPtr = reinterpret_cast<double*>(OutTargetPose.LocalTransformsView.Rotations.GetData());
+
+	// Our SoA buffer is contiguous
+	// Because the translations/scales have the same size (FVector), each entry is a fixed offset apart and we can use a single ptr/offset pair
+	const double* RESTRICT SrcTranslationPtr = reinterpret_cast<const double*>(SourcePose.LocalTransformsView.Translations.GetData());
+	double* RESTRICT DstTranslationPtr = reinterpret_cast<double*>(OutTargetPose.LocalTransformsView.Translations.GetData());
+
+	const int64 SrcScaleOffset = reinterpret_cast<const double*>(SourcePose.LocalTransformsView.Scales3D.GetData()) - SrcTranslationPtr;
+	const int64 DstScaleOffset = reinterpret_cast<double*>(OutTargetPose.LocalTransformsView.Scales3D.GetData()) - DstTranslationPtr;
+
+	const FRemappedBone* RESTRICT RemappedBonePtr = SourceToTargetBoneIndexMap.GetData();
+	const FRemappedBone* RESTRICT RemappedBoneEndPtr = RemappedBonePtr + NumToRemap;
+
+	while (RemappedBonePtr < RemappedBoneEndPtr)
+	{
+		const int32 SourceBoneIndex = RemappedBonePtr->SourceBoneIndex;
+		const int32 TargetBoneIndex = RemappedBonePtr->TargetBoneIndex;
+
+		SourcePose.LocalTransformsView.Rotations.RangeCheck(SourceBoneIndex);
+		OutTargetPose.LocalTransformsView.Rotations.RangeCheck(TargetBoneIndex);
+
+		VectorRegister4Double Rotation = VectorLoadAligned(SrcRotationPtr + SourceBoneIndex * 4);
+		VectorRegister4Double Translation = VectorLoadAligned(SrcTranslationPtr + SourceBoneIndex * 3);
+		VectorRegister4Double Scale = VectorLoadAligned(SrcTranslationPtr + SourceBoneIndex * 3 + SrcScaleOffset);
+
+		VectorStoreAligned(Rotation, DstRotationPtr + TargetBoneIndex * 4);
+		VectorStoreFloat3(Translation, DstTranslationPtr + TargetBoneIndex * 3);
+		VectorStoreFloat3(Scale, DstTranslationPtr + TargetBoneIndex * 3 + DstScaleOffset);
+
+		RemappedBonePtr++;
+	}
+#else
+	for (int32 MappingElementIndex = 0; MappingElementIndex < NumToRemap; ++MappingElementIndex)
 	{
 		const int32 SourceBoneIndex = SourceToTargetBoneIndexMap[MappingElementIndex].SourceBoneIndex;
 		const int32 TargetBoneIndex = SourceToTargetBoneIndexMap[MappingElementIndex].TargetBoneIndex;
@@ -153,6 +188,7 @@ void FRemapPoseData::RemapPose(const UE::AnimNext::FLODPoseHeap& SourcePose,
 		const auto& TargetTransform = SourcePose.LocalTransformsView[SourceBoneIndex];
 		OutTargetPose.LocalTransforms[TargetBoneIndex] = TargetTransform;
 	}
+#endif
 
 	// For the cases where the target skeleton does not share the same root bone with the source, try to find the corresponding
 	// bone on the source skeleton for the target's root bone and sync the skeletons up from there by calculating the delta transform
diff --git a/Engine/Plugins/Experimental/AnimNext/Source/AnimNext/Private/Graph/RemapPoseData.h b/Engine/Plugins/Experimental/AnimNext/Source/AnimNext/Private/Graph/RemapPoseData.h
--- a/Engine/Plugins/Experimental/AnimNext/Source/AnimNext/Private/Graph/RemapPoseData.h
+++ b/Engine/Plugins/Experimental/AnimNext/Source/AnimNext/Private/Graph/RemapPoseData.h
@@ -5,6 +5,7 @@
 #include "Animation/AttributesRuntime.h"
 #include "ReferencePose.h"
 #include "LODPose.h"
+#include "BoneIndices.h"
 
 #include "RemapPoseData.generated.h"
 
@@ -13,8 +14,8 @@ struct FRemappedBone
 {
 	GENERATED_BODY()
 
-	int32 SourceBoneIndex = INDEX_NONE;
-	int32 TargetBoneIndex = INDEX_NONE;
+	FBoneIndexType SourceBoneIndex = INDEX_NONE;
+	FBoneIndexType TargetBoneIndex = INDEX_NONE;
 };
 
 USTRUCT()
