AnimNext: Remove module write guard on module handle release as it is safe to unregister tick functions from the GT

diff --git a/Engine/Plugins/Experimental/AnimNext/Source/AnimNext/Private/Component/AnimNextWorldSubsystem.cpp b/Engine/Plugins/Experimental/AnimNext/Source/AnimNext/Private/Component/AnimNextWorldSubsystem.cpp
--- a/Engine/Plugins/Experimental/AnimNext/Source/AnimNext/Private/Component/AnimNextWorldSubsystem.cpp
+++ b/Engine/Plugins/Experimental/AnimNext/Source/AnimNext/Private/Component/AnimNextWorldSubsystem.cpp
@@ -208,7 +208,7 @@ void UAnimNextWorldSubsystem::UnregisterHandle(UE::AnimNext::FModuleHandle& InOu
 	{
 		FWriteScopeLock InstancesLockScope(InstancesLock);
 		FAnimNextModuleInstance& Instance = Instances.Get(InOutHandle);
-		FModuleWriteGuard Guard(&Instance);
+
 #if UE_ENABLE_DEBUG_DRAWING
 		// Remove debug drawing immediately as the renderer will need to know about this before EOF
 		Instance.DebugDraw->RemovePrimitive();
