- Fixed up retargetting using wrong index remapping
- Added TODO for RetargetBoneTransform assumption of NumLOD == 1 || LODIndex == 0
- Added ensure message for any orient/scale retargetting in AnimNext

diff --git a/Engine/Plugins/Experimental/AnimNext/Source/AnimNext/Private/DecompressionTools.cpp b/Engine/Plugins/Experimental/AnimNext/Source/AnimNext/Private/DecompressionTools.cpp
--- a/Engine/Plugins/Experimental/AnimNext/Source/AnimNext/Private/DecompressionTools.cpp
+++ b/Engine/Plugins/Experimental/AnimNext/Source/AnimNext/Private/DecompressionTools.cpp
@@ -310,7 +310,7 @@ void FDecompressionTools::DecompressPose(FLODPose& OutAnimationPoseData,
 										const FRootMotionReset& RootMotionReset)
 {
 	const FReferencePose& ReferencePose = OutAnimationPoseData.GetRefPose();
-	const TArrayView<const FBoneIndexType> LODBoneIndexToSkeletonBoneIndexMap = OutAnimationPoseData.GetLODBoneIndexToMeshBoneIndexMap();
+	const TArrayView<const FBoneIndexType> LODBoneIndexToSkeletonBoneIndexMap = OutAnimationPoseData.GetLODBoneIndexToSkeletonBoneIndexMap();
 	const TArrayView<const FBoneIndexType> SkeletonToLODBoneIndexes = ReferencePose.GetSkeletonBoneIndexToLODBoneIndexMap(); // Full list of Skeleton to LOD conversion
 	const int32 NumLODBoneIndexes = LODBoneIndexToSkeletonBoneIndexMap.Num();
 
@@ -497,7 +497,7 @@ void FDecompressionTools::DecompressPose(FLODPose& OutAnimationPoseData,
 		{
 			for (int32 LODBoneIndex = (bFirstTrackIsRootBone ? 1 : 0); LODBoneIndex < NumLODBoneIndexes; ++LODBoneIndex)
 			{
-				const int32 TargetSkeletonBoneIndex = ReferencePose.GetSkeletonBoneIndexFromLODBoneIndex(LODBoneIndex);
+				const int32 TargetSkeletonBoneIndex = LODBoneIndexToSkeletonBoneIndexMap[LODBoneIndex];
 
 				// Mesh space additives do not require fix-up
 				if (DecompressionContext.GetAdditiveType() == AAT_LocalSpaceBase)
@@ -525,7 +525,7 @@ void FDecompressionTools::DecompressPose(FLODPose& OutAnimationPoseData,
 		{
 			for (int32 LODBoneIndex = (bFirstTrackIsRootBone ? 1 : 0); LODBoneIndex < NumLODBoneIndexes; ++LODBoneIndex)
 			{
-				const int32 TargetSkeletonBoneIndex = LODBoneIndexToSkeletonBoneIndexMap[LODBoneIndex]; // ReferencePose.GetSkeletonBoneIndexFromLODBoneIndex(LODBoneIndex);
+				const int32 TargetSkeletonBoneIndex = LODBoneIndexToSkeletonBoneIndexMap[LODBoneIndex];
 				OutAnimationPoseData.LocalTransformsView[LODBoneIndex].SetRotation(SkeletonRemapping.RetargetBoneRotationToTargetSkeleton(TargetSkeletonBoneIndex, OutAnimationPoseData.LocalTransformsView[LODBoneIndex].GetRotation()));
 
 				// Check what retarget mode to use for the translational retargeting for this specific bone.
@@ -645,6 +645,7 @@ void FDecompressionTools::DecompressPose(FLODPose& OutAnimationPoseData,
 	//		}
 	//	}
 	//}
+	ensureMsgf(OrientAndScaleRetargetingPairs.Num() == 0, TEXT("OrientAndScaleRetargetingPairs found for %s (%s -> %s), Orient and Scale retargeting is not yet supported in AnimNext"), *DecompressionContext.AnimName.ToString(), SourceSkeleton ? *SourceSkeleton->GetName() : TEXT("None"), TargetSkeleton ? *TargetSkeleton->GetName() : TEXT("None"));
 
 	if (bIsMeshSpaceAdditive)
 	{
@@ -706,6 +707,7 @@ void FDecompressionTools::RetargetBoneTransform(const FReferencePose& ReferenceP
 		//const USkeleton* TargetSkeleton = RequiredBones.GetSkeletonAsset();
 		const FSkeletonRemapping& SkeletonRemapping = UE::Anim::FSkeletonRemappingRegistry::Get().GetRemapping(SourceSkeleton, TargetSkeleton);
 
+		// TODO this assumes ReferencePose to have single LOD entry, or target LOD index being 0 
 		const int32 TargetSkeletonBoneIndex = ReferencePose.GetSkeletonBoneIndexFromLODBoneIndex(LODBoneIndex);
 		const int32 SourceSkeletonBoneIndex = SkeletonRemapping.IsValid() ? SkeletonRemapping.GetSourceSkeletonBoneIndex(TargetSkeletonBoneIndex) : SkeletonBoneIndex;
 
