AnimNext: Fix ensure on map reload/shutdown

diff --git a/Engine/Plugins/Experimental/AnimNext/Source/AnimNext/Private/Component/AnimNextWorldSubsystem.cpp b/Engine/Plugins/Experimental/AnimNext/Source/AnimNext/Private/Component/AnimNextWorldSubsystem.cpp
--- a/Engine/Plugins/Experimental/AnimNext/Source/AnimNext/Private/Component/AnimNextWorldSubsystem.cpp
+++ b/Engine/Plugins/Experimental/AnimNext/Source/AnimNext/Private/Component/AnimNextWorldSubsystem.cpp
@@ -47,6 +47,7 @@ void UAnimNextWorldSubsystem::BeginDestroy()
 			FRWScopeLock InstancesLockScope(InstancesLock, SLT_Write);
 			for (FAnimNextModuleInstance& Instance : Instances)
 			{
+				Instance.RemoveAllTickDependencies();
 				Instance.Uninitialize();
 			}
 			Instances = UE::AnimNext::TPool<FAnimNextModuleInstance>(); // Force instances destruction
diff --git a/Engine/Plugins/Experimental/AnimNext/Source/AnimNext/Private/Module/AnimNextModuleInstance.cpp b/Engine/Plugins/Experimental/AnimNext/Source/AnimNext/Private/Module/AnimNextModuleInstance.cpp
--- a/Engine/Plugins/Experimental/AnimNext/Source/AnimNext/Private/Module/AnimNextModuleInstance.cpp
+++ b/Engine/Plugins/Experimental/AnimNext/Source/AnimNext/Private/Module/AnimNextModuleInstance.cpp
@@ -362,6 +362,19 @@ void FAnimNextModuleInstance::RemoveAllTickDependencies()
 	{
 		TickFunction.RemoveAllExternalSubsequents();
 	}
+
+	if(Pool)
+	{
+		TArray<FPrerequisiteReference, TInlineAllocator<8>> PrerequisiteRefsCopy;
+		PrerequisiteRefsCopy.Append(PrerequisiteRefs);
+		for (const FPrerequisiteReference& PrerequisiteHandle : PrerequisiteRefsCopy)
+		{
+			if(FAnimNextModuleInstance* PrerequisiteInstance = Pool->TryGet(PrerequisiteHandle.Handle))
+			{
+				RemovePrerequisite(*PrerequisiteInstance);
+			}
+		}
+	}
 }
 
 void FAnimNextModuleInstance::ResetBindingsAndInstanceData()
