AnimNext: Remove ExecuteBindings_WT tick function

Instead, just run the bindings event before the first user event if its flag is set.
Reduces total tick function count, further reducing queuing/tick overhead time.
Allows external dependencies to rely on dependencies with user events only (e.g. depending on PrePhysics will behave as expected).

diff --git a/Engine/Plugins/Experimental/AnimNext/Source/AnimNext/Internal/Module/ModuleTickFunction.h b/Engine/Plugins/Experimental/AnimNext/Source/AnimNext/Internal/Module/ModuleTickFunction.h
--- a/Engine/Plugins/Experimental/AnimNext/Source/AnimNext/Internal/Module/ModuleTickFunction.h
+++ b/Engine/Plugins/Experimental/AnimNext/Source/AnimNext/Internal/Module/ModuleTickFunction.h
@@ -36,6 +36,7 @@ struct FModuleEventTickFunction : public FTickFunction
 	virtual FString DiagnosticMessage() override;
 
 	void Initialize(float DeltaTime);
+	void ExecuteBindings_WT(float DeltaTime);
 	void Run(float DeltaTime);
 	void EndTick(float DeltaTime);
 
@@ -63,6 +64,7 @@ struct FModuleEventTickFunction : public FTickFunction
 	uint8 bLastUserEvent : 1 = false;
 	uint8 bFirstUserEvent : 1 = false;
 	uint8 bUserEvent : 1 = false;
+	uint8 bRunBindingsEvent : 1 = false;
 
 	// External dependencies that this tick function needs to unregister when it is destroyed
 	TArray<FTickPrerequisite> ExternalSubsequents;
diff --git a/Engine/Plugins/Experimental/AnimNext/Source/AnimNext/Internal/Module/RigUnit_AnimNextModuleEvents.h b/Engine/Plugins/Experimental/AnimNext/Source/AnimNext/Internal/Module/RigUnit_AnimNextModuleEvents.h
--- a/Engine/Plugins/Experimental/AnimNext/Source/AnimNext/Internal/Module/RigUnit_AnimNextModuleEvents.h
+++ b/Engine/Plugins/Experimental/AnimNext/Source/AnimNext/Internal/Module/RigUnit_AnimNextModuleEvents.h
@@ -97,6 +97,7 @@ struct ANIMNEXT_API FRigUnit_AnimNextExecuteBindings_WT : public FRigUnit_AnimNe
 	virtual FName GetEventName() const override final { return EventName; }
 	virtual int32 GetSortOrder() const override final { return 1; }	// Sort after GT bindings
 	virtual bool IsUserEvent() const override final { return false; }
+	virtual bool IsTask() const override final { return false; }
 
 	// FRigUnit_AnimNextModuleEventBase interface
 	virtual UE::AnimNext::EModuleEventPhase GetEventPhase() const override final { return UE::AnimNext::EModuleEventPhase::PreExecute; }
diff --git a/Engine/Plugins/Experimental/AnimNext/Source/AnimNext/Private/Module/AnimNextModuleInstance.cpp b/Engine/Plugins/Experimental/AnimNext/Source/AnimNext/Private/Module/AnimNextModuleInstance.cpp
--- a/Engine/Plugins/Experimental/AnimNext/Source/AnimNext/Private/Module/AnimNextModuleInstance.cpp
+++ b/Engine/Plugins/Experimental/AnimNext/Source/AnimNext/Private/Module/AnimNextModuleInstance.cpp
@@ -240,6 +240,14 @@ void FAnimNextModuleInstance::Initialize()
 			{
 				TickFunction.bFirstUserEvent = true;
 				bFoundFirstUserEvent = true;
+
+				// Set this first user event to run the bindings event, if it exists
+				auto IsExecuteBindingsEvent = [](const Private::FImplementedModuleEvent& InEvent)
+				{
+					return InEvent.EventName == FRigUnit_AnimNextExecuteBindings_WT::EventName;
+				};
+
+				TickFunction.bRunBindingsEvent = ImplementedModuleEvents.ContainsByPredicate(IsExecuteBindingsEvent);
 			}
 		}
 
diff --git a/Engine/Plugins/Experimental/AnimNext/Source/AnimNext/Private/Module/ModuleTickFunction.cpp b/Engine/Plugins/Experimental/AnimNext/Source/AnimNext/Private/Module/ModuleTickFunction.cpp
--- a/Engine/Plugins/Experimental/AnimNext/Source/AnimNext/Private/Module/ModuleTickFunction.cpp
+++ b/Engine/Plugins/Experimental/AnimNext/Source/AnimNext/Private/Module/ModuleTickFunction.cpp
@@ -24,6 +24,11 @@ void FModuleEventTickFunction::Initialize(float InDeltaTime)
 	ModuleInstance->RunRigVMEvent(FRigUnit_AnimNextInitializeEvent::EventName, InDeltaTime);
 }
 
+void FModuleEventTickFunction::ExecuteBindings_WT(float InDeltaTime)
+{
+	ModuleInstance->RunRigVMEvent(FRigUnit_AnimNextExecuteBindings_WT::EventName, InDeltaTime);
+}
+
 void FModuleEventTickFunction::EndTick(float DeltaTime)
 {
 	SCOPED_NAMED_EVENT(AnimNext_Module_EndTick, FColor::Orange);
@@ -119,11 +124,20 @@ void FModuleEventTickFunction::Run(float InDeltaTime)
 
 	FModuleWriteGuard Guard(ModuleInstance);
 
-	// Run the pending initialize if required
-	if (bFirstUserEvent && ModuleInstance->InitState == FAnimNextModuleInstance::EInitState::PendingInitializeEvent)
+	if (bFirstUserEvent)
 	{
-		Initialize(InDeltaTime);
-		ModuleInstance->TransitionToInitState(FAnimNextModuleInstance::EInitState::FirstUpdate);
+		if (bRunBindingsEvent)
+		{
+			// Execute any WT bindings
+			ExecuteBindings_WT(InDeltaTime);
+		}
+
+		// Run the pending initialize if required
+		if (ModuleInstance->InitState == FAnimNextModuleInstance::EInitState::PendingInitializeEvent)
+		{
+			Initialize(InDeltaTime);
+			ModuleInstance->TransitionToInitState(FAnimNextModuleInstance::EInitState::FirstUpdate);
+		}
 	}
 
 	while (!PreExecuteTasks.IsEmpty())
