Optimize SoA pose copy

diff --git a/Engine/Plugins/Experimental/AnimNext/Source/AnimNext/Private/TransformArrayOperations_SoA.cpp b/Engine/Plugins/Experimental/AnimNext/Source/AnimNext/Private/TransformArrayOperations_SoA.cpp
--- a/Engine/Plugins/Experimental/AnimNext/Source/AnimNext/Private/TransformArrayOperations_SoA.cpp
+++ b/Engine/Plugins/Experimental/AnimNext/Source/AnimNext/Private/TransformArrayOperations_SoA.cpp
@@ -80,11 +80,51 @@ namespace UE::AnimNext
 		check(StartIndex >= 0 && StartIndex <= NumTransforms);
 		check(EndIndex <= NumTransforms);
 
-		for (int32 TransformIndex = StartIndex; TransformIndex < EndIndex; ++TransformIndex)
+		const double* RESTRICT SrcRotationPtr = reinterpret_cast<const double*>(Source.Rotations.GetData()) + (StartIndex * 4);
+		const double* RESTRICT SrcRotationEndPtr = reinterpret_cast<const double*>(Source.Rotations.GetData()) + (EndIndex * 4);
+		double* RESTRICT DstRotationPtr = reinterpret_cast<double*>(Dest.Rotations.GetData()) + (StartIndex * 4);
+
+		// Our SoA buffer is contiguous
+		// Because the translations/scales have the same size (FVector), each entry is a fixed offset apart and we can use a single ptr/offset pair
+		const double* RESTRICT SrcTranslationPtr = reinterpret_cast<const double*>(Source.Translations.GetData()) + (StartIndex * 3);
+		double* RESTRICT DstTranslationPtr = reinterpret_cast<double*>(Dest.Translations.GetData()) + (StartIndex * 3);
+
+		const int64 SrcScaleOffset = reinterpret_cast<const double*>(Source.Scales3D.GetData()) - SrcTranslationPtr;
+		const int64 DstScaleOffset = reinterpret_cast<double*>(Dest.Scales3D.GetData()) - DstTranslationPtr;
+
+		// Iterate until the last element
+		// Each translation/scale is 24 bytes (FVector) and everything is contiguous
+		// For all of these elements, we can use aligned 32-bytes load/stores for translation/scale
+		// since the next entry will simply overwrite the last 8 bytes we trashed
+		// We effectively copy 24 bytes each iteration for translation/scale and we overwrite
+		// the last 8 bytes of the next entry
+		const double* RESTRICT SrcRotationLastPtr = SrcRotationEndPtr - 4;
+		while (SrcRotationPtr < SrcRotationLastPtr)
 		{
-			Dest.Translations[TransformIndex] = Source.Translations[TransformIndex];
-			Dest.Rotations[TransformIndex] = Source.Rotations[TransformIndex];
-			Dest.Scales3D[TransformIndex] = Source.Scales3D[TransformIndex];
+			VectorRegister4Double Rotation = VectorLoadAligned(SrcRotationPtr);
+			VectorRegister4Double Translation = VectorLoadAligned(SrcTranslationPtr);
+			VectorRegister4Double Scale = VectorLoadAligned(SrcTranslationPtr + SrcScaleOffset);
+
+			VectorStoreAligned(Rotation, DstRotationPtr);
+			VectorStoreAligned(Translation, DstTranslationPtr);
+			VectorStoreAligned(Scale, DstTranslationPtr + DstScaleOffset);
+
+			SrcRotationPtr += 4;
+			SrcTranslationPtr += 3;
+			DstRotationPtr += 4;
+			DstTranslationPtr += 3;
+		}
+
+		// The last iteration is performed afterwards to avoid stomping past the end of our copy
+		if (SrcRotationPtr < SrcRotationEndPtr)
+		{
+			VectorRegister4Double Rotation = VectorLoadAligned(SrcRotationPtr);
+			VectorRegister4Double Translation = VectorLoadAligned(SrcTranslationPtr);
+			VectorRegister4Double Scale = VectorLoadAligned(SrcTranslationPtr + SrcScaleOffset);
+
+			VectorStoreAligned(Rotation, DstRotationPtr);
+			VectorStoreFloat3(Translation, DstTranslationPtr);
+			VectorStoreFloat3(Scale, DstTranslationPtr + DstScaleOffset);
 		}
 	}
 
