[Truncation Warnings] Update EdGraphSchema, EdGraphActions, and SGraphEditor to use FVector2f

diff --git a/Engine/Plugins/Animation/DeformerGraph/Source/OptimusEditor/Private/OptimusEditor.cpp b/Engine/Plugins/Animation/DeformerGraph/Source/OptimusEditor/Private/OptimusEditor.cpp
--- a/Engine/Plugins/Animation/DeformerGraph/Source/OptimusEditor/Private/OptimusEditor.cpp
+++ b/Engine/Plugins/Animation/DeformerGraph/Source/OptimusEditor/Private/OptimusEditor.cpp
@@ -293,7 +293,7 @@ bool FOptimusEditor::SetEditGraph(UOptimusNodeGraph* InNodeGraph)
 		float Zoom;
 		if (InNodeGraph->GetViewLocationAndZoom(Location, Zoom))
 		{
-			GraphEditorWidget->SetViewLocation(Location, Zoom);
+			GraphEditorWidget->SetViewLocation(UE::Slate::FDeprecateVector2DParameter(Location), Zoom);
 		}
 		else
 		{
@@ -565,7 +565,7 @@ void FOptimusEditor::PasteNodes() const
 
 	UOptimusNodeGraph* ModelGraph = EditorGraph->GetModelGraph();
 	ModelGraph->DuplicateNodes(
-		TransientGraph->GetAllNodes(), GraphEditorWidget->GetPasteLocation(), TEXT("Paste"));
+		TransientGraph->GetAllNodes(), GraphEditorWidget->GetPasteLocation2f(), TEXT("Paste"));
 }
 
 
@@ -585,7 +585,7 @@ void FOptimusEditor::DuplicateNodes() const
 	UOptimusNodeGraph* ModelGraph = EditorGraph->GetModelGraph();
 	const TArray<UOptimusNode*> ModelNodes = GetSelectedModelNodes();
 
-	ModelGraph->DuplicateNodes(ModelNodes, GraphEditorWidget->GetPasteLocation());
+	ModelGraph->DuplicateNodes(ModelNodes, GraphEditorWidget->GetPasteLocation2f());
 }
 
 
@@ -802,7 +802,7 @@ bool FOptimusEditor::CanConvertToSubGraph() const
 void FOptimusEditor::OnCreateComment()
 {
 	// In case there aren't any nodes selected, create a comment node at current mouse position
-	const FVector2D CurrentMouseGraphPosition = GraphEditorWidget->GetPasteLocation();
+	const FVector2f CurrentMouseGraphPosition = GraphEditorWidget->GetPasteLocation2f();
 	FOptimusGraphSchemaAction_NewCommentNode NewCommentNodeAction;
 	NewCommentNodeAction.PerformAction(EditorGraph, nullptr, CurrentMouseGraphPosition, true);
 }
@@ -1127,7 +1127,7 @@ void FOptimusEditor::CreateWidgets()
 {
 	// -- Graph editor
 	GraphEditorWidget = CreateGraphEditorWidget();
-	GraphEditorWidget->SetViewLocation(FVector2D::ZeroVector, 1);
+	GraphEditorWidget->SetViewLocation(FVector2f::ZeroVector, 1);
 
 	FindResults =
 		SNew(SFindInDeformer, SharedThis(this));
@@ -1339,11 +1339,11 @@ void FOptimusEditor::StoreCurrentViewLocation()
 {
 	if (EditorGraph->NodeGraph)
 	{
-		FVector2D Location;
+		FVector2f Location;
 		float Zoom;
 		
 		GraphEditorWidget->GetViewLocation(Location, Zoom);
-		EditorGraph->NodeGraph->SetViewLocationAndZoom(Location, Zoom);
+		EditorGraph->NodeGraph->SetViewLocationAndZoom(FDeprecateSlateVector2D(Location), Zoom);
 	}
 }
 
diff --git a/Engine/Plugins/Animation/DeformerGraph/Source/OptimusEditor/Private/OptimusEditorGraphSchemaActions.cpp b/Engine/Plugins/Animation/DeformerGraph/Source/OptimusEditor/Private/OptimusEditorGraphSchemaActions.cpp
--- a/Engine/Plugins/Animation/DeformerGraph/Source/OptimusEditor/Private/OptimusEditorGraphSchemaActions.cpp
+++ b/Engine/Plugins/Animation/DeformerGraph/Source/OptimusEditor/Private/OptimusEditorGraphSchemaActions.cpp
@@ -20,7 +20,7 @@
 UEdGraphNode* FOptimusGraphSchemaAction_NewNode::PerformAction(
 	UEdGraph* InParentGraph, 
 	UEdGraphPin* InFromPin, 
-	const FVector2D InLocation, 
+	const FVector2f& InLocation, 
 	bool bInSelectNewNode /*= true*/
 	)
 {
@@ -28,7 +28,7 @@ UEdGraphNode* FOptimusGraphSchemaAction_NewNode::PerformAction(
 	
 	if (ensure(Graph != nullptr) && ensure(NodeClass != nullptr))
 	{
-		UOptimusNode* ModelNode = Graph->GetModelGraph()->AddNode(NodeClass, InLocation);
+		UOptimusNode* ModelNode = Graph->GetModelGraph()->AddNode(NodeClass, FDeprecateSlateVector2D(InLocation));
 
  		// FIXME: Automatic connection from the given pin.
 
@@ -47,7 +47,7 @@ UEdGraphNode* FOptimusGraphSchemaAction_NewNode::PerformAction(
 UEdGraphNode* FOptimusGraphSchemaAction_NewConstantValueNode::PerformAction(
 	UEdGraph* InParentGraph,
 	UEdGraphPin* InFromPin,
-	const FVector2D InLocation,
+	const FVector2f& InLocation,
 	bool bInSelectNewNode
 	)
 {
@@ -55,7 +55,7 @@ UEdGraphNode* FOptimusGraphSchemaAction_NewConstantValueNode::PerformAction(
 	
 	if (ensure(Graph != nullptr) && ensure(DataType.IsValid()))
 	{
-		UOptimusNode* ModelNode = Graph->GetModelGraph()->AddValueNode(DataType, InLocation);
+		UOptimusNode* ModelNode = Graph->GetModelGraph()->AddValueNode(DataType, FDeprecateSlateVector2D(InLocation));
 
 		// FIXME: Automatic connection from the given pin.
 
@@ -74,7 +74,7 @@ UEdGraphNode* FOptimusGraphSchemaAction_NewConstantValueNode::PerformAction(
 UEdGraphNode* FOptimusGraphSchemaAction_NewDataInterfaceNode::PerformAction(
 	UEdGraph* InParentGraph,
 	UEdGraphPin* InFromPin,
-	const FVector2D InLocation,
+	const FVector2f& InLocation,
 	bool bInSelectNewNode
 	)
 {
@@ -82,7 +82,7 @@ UEdGraphNode* FOptimusGraphSchemaAction_NewDataInterfaceNode::PerformAction(
 	
 	if (ensure(Graph != nullptr) && ensure(DataInterfaceClass != nullptr))
 	{
-		UOptimusNode* ModelNode = Graph->GetModelGraph()->AddDataInterfaceNode(DataInterfaceClass, InLocation);
+		UOptimusNode* ModelNode = Graph->GetModelGraph()->AddDataInterfaceNode(DataInterfaceClass, FDeprecateSlateVector2D(InLocation));
 
 		// FIXME: Automatic connection from the given pin.
 
@@ -98,13 +98,13 @@ UEdGraphNode* FOptimusGraphSchemaAction_NewDataInterfaceNode::PerformAction(
 }
 
 UEdGraphNode* FOptimusGraphSchemaAction_NewLoopTerminalNodes::PerformAction(UEdGraph* InParentGraph, UEdGraphPin* FromPin,
-	const FVector2D Location, bool bInSelectNewNode)
+	const FVector2f& Location, bool bInSelectNewNode)
 {
 	UOptimusEditorGraph* Graph = Cast<UOptimusEditorGraph>(InParentGraph);
 	
 	if (ensure(Graph != nullptr))
 	{
-		TArray<UOptimusNode*> Nodes = Graph->GetModelGraph()->AddLoopTerminalNodes(Location);
+		TArray<UOptimusNode*> Nodes = Graph->GetModelGraph()->AddLoopTerminalNodes(FDeprecateSlateVector2D(Location));
 
 		if (ensure(Nodes.Num() == 2))
 		{
@@ -120,14 +120,14 @@ UEdGraphNode* FOptimusGraphSchemaAction_NewLoopTerminalNodes::PerformAction(UEdG
 	return nullptr;
 }
 
-UEdGraphNode* FOptimusGraphSchemaAction_NewCommentNode::PerformAction(UEdGraph* InParentGraph, UEdGraphPin* FromPin, const FVector2D Location,
+UEdGraphNode* FOptimusGraphSchemaAction_NewCommentNode::PerformAction(UEdGraph* InParentGraph, UEdGraphPin* FromPin, const FVector2f& Location,
 	bool bInSelectNewNode)
 {
 	UOptimusEditorGraph* Graph = Cast<UOptimusEditorGraph>(InParentGraph);
 	
 	if (ensure(Graph != nullptr))
 	{
-		UOptimusNode* ModelNode = OptimusEditor::CreateCommentNode(Graph, Location);
+		UOptimusNode* ModelNode = OptimusEditor::CreateCommentNode(Graph, FDeprecateSlateVector2D(Location));
 
 		UEdGraphNode* GraphNode = Graph->FindGraphNodeFromModelNode(ModelNode);
 		if (GraphNode && bInSelectNewNode)
@@ -140,7 +140,7 @@ UEdGraphNode* FOptimusGraphSchemaAction_NewCommentNode::PerformAction(UEdGraph*
 	return nullptr;
 }
 
-UEdGraphNode* FOptimusGraphSchemaAction_NewFunctionReferenceNode::PerformAction(UEdGraph* InParentGraph, UEdGraphPin* InFromPin, const FVector2D InLocation, bool bInSelectNewNode)
+UEdGraphNode* FOptimusGraphSchemaAction_NewFunctionReferenceNode::PerformAction(UEdGraph* InParentGraph, UEdGraphPin* InFromPin, const FVector2f& InLocation, bool bInSelectNewNode)
 {
 	UOptimusEditorGraph* Graph = Cast<UOptimusEditorGraph>(InParentGraph);
 	
@@ -150,7 +150,7 @@ UEdGraphNode* FOptimusGraphSchemaAction_NewFunctionReferenceNode::PerformAction(
 		{
 			if (UOptimusFunctionNodeGraph* FunctionGraph = Deformer->FindFunctionByGuid(FunctionGraphGuid))
 			{
-				UOptimusNode* ModelNode = Graph->GetModelGraph()->AddFunctionReferenceNode(FunctionGraph, InLocation);
+				UOptimusNode* ModelNode = Graph->GetModelGraph()->AddFunctionReferenceNode(FunctionGraph, FDeprecateSlateVector2D(InLocation));
 
 				// FIXME: Automatic connection from the given pin.
 
diff --git a/Engine/Plugins/Animation/DeformerGraph/Source/OptimusEditor/Private/OptimusEditorGraphSchemaActions.h b/Engine/Plugins/Animation/DeformerGraph/Source/OptimusEditor/Private/OptimusEditorGraphSchemaActions.h
--- a/Engine/Plugins/Animation/DeformerGraph/Source/OptimusEditor/Private/OptimusEditorGraphSchemaActions.h
+++ b/Engine/Plugins/Animation/DeformerGraph/Source/OptimusEditor/Private/OptimusEditorGraphSchemaActions.h
@@ -48,7 +48,7 @@ struct FOptimusGraphSchemaAction_NewNode :
 	FName GetTypeId() const override { return StaticGetTypeId(); }
 
 	// FEdGraphSchemaAction overrides
-	UEdGraphNode* PerformAction(class UEdGraph* ParentGraph, UEdGraphPin* FromPin, const FVector2D Location, bool bSelectNewNode = true) override;
+	UEdGraphNode* PerformAction(class UEdGraph* ParentGraph, UEdGraphPin* FromPin, const FVector2f& Location, bool bSelectNewNode = true) override;
 };
 
 USTRUCT()
@@ -67,7 +67,7 @@ struct FOptimusGraphSchemaAction_NewConstantValueNode :
 	FName GetTypeId() const override { return StaticGetTypeId(); }
 
 	// FEdGraphSchemaAction overrides
-	UEdGraphNode* PerformAction(class UEdGraph* ParentGraph, UEdGraphPin* FromPin, const FVector2D Location, bool bSelectNewNode = true) override;
+	UEdGraphNode* PerformAction(class UEdGraph* ParentGraph, UEdGraphPin* FromPin, const FVector2f& Location, bool bSelectNewNode = true) override;
 };
 
 USTRUCT()
@@ -86,7 +86,7 @@ struct FOptimusGraphSchemaAction_NewDataInterfaceNode :
 	FName GetTypeId() const override { return StaticGetTypeId(); }
 
 	// FEdGraphSchemaAction overrides
-	UEdGraphNode* PerformAction(class UEdGraph* ParentGraph, UEdGraphPin* FromPin, const FVector2D Location, bool bInSelectNewNode = true) override;
+	UEdGraphNode* PerformAction(class UEdGraph* ParentGraph, UEdGraphPin* FromPin, const FVector2f& Location, bool bInSelectNewNode = true) override;
 };
 
 
@@ -102,7 +102,7 @@ struct FOptimusGraphSchemaAction_NewLoopTerminalNodes : public FEdGraphSchemaAct
 	FName GetTypeId() const override { return StaticGetTypeId(); }
 
 	// FEdGraphSchemaAction overrides
-	UEdGraphNode* PerformAction(class UEdGraph* InParentGraph, UEdGraphPin* FromPin, const FVector2D Location, bool bInSelectNewNode = true) override;
+	UEdGraphNode* PerformAction(class UEdGraph* InParentGraph, UEdGraphPin* FromPin, const FVector2f& Location, bool bInSelectNewNode = true) override;
 };
 
 USTRUCT()
@@ -117,7 +117,8 @@ struct FOptimusGraphSchemaAction_NewCommentNode : public FEdGraphSchemaAction
 	FName GetTypeId() const override { return StaticGetTypeId(); }
 
 	// FEdGraphSchemaAction overrides
-	UEdGraphNode* PerformAction(class UEdGraph* InParentGraph, UEdGraphPin* FromPin, const FVector2D Location, bool bInSelectNewNode = true) override;
+	using FEdGraphSchemaAction::PerformAction;
+	UEdGraphNode* PerformAction(class UEdGraph* InParentGraph, UEdGraphPin* FromPin, const FVector2f& Location, bool bInSelectNewNode = true) override;
 };
 
 USTRUCT()
@@ -137,7 +138,7 @@ struct FOptimusGraphSchemaAction_NewFunctionReferenceNode:
 	FName GetTypeId() const override { return StaticGetTypeId(); }
 
 	// FEdGraphSchemaAction overrides
-	UEdGraphNode* PerformAction(class UEdGraph* InParentGraph, UEdGraphPin* InFromPin, const FVector2D InLocation, bool bInSelectNewNode = true) override;
+	UEdGraphNode* PerformAction(class UEdGraph* InParentGraph, UEdGraphPin* InFromPin, const FVector2f& InLocation, bool bInSelectNewNode = true) override;
 };
 
 /// Reference to a UOptimusNodeGraph.
diff --git a/Engine/Plugins/Cameras/GameplayCameras/Source/GameplayCamerasEditor/Private/Editors/CameraNodeGraphDragDropOp.cpp b/Engine/Plugins/Cameras/GameplayCameras/Source/GameplayCamerasEditor/Private/Editors/CameraNodeGraphDragDropOp.cpp
--- a/Engine/Plugins/Cameras/GameplayCameras/Source/GameplayCamerasEditor/Private/Editors/CameraNodeGraphDragDropOp.cpp
+++ b/Engine/Plugins/Cameras/GameplayCameras/Source/GameplayCamerasEditor/Private/Editors/CameraNodeGraphDragDropOp.cpp
@@ -38,7 +38,7 @@ FReply FCameraNodeGraphInterfaceParameterDragDropOp::ExecuteDragOver(TSharedPtr<
 	return FReply::Handled();
 }
 
-FReply FCameraNodeGraphInterfaceParameterDragDropOp::ExecuteDrop(TSharedPtr<SGraphEditor> GraphEditor, const FVector2D& NewLocation)
+FReply FCameraNodeGraphInterfaceParameterDragDropOp::ExecuteDrop(TSharedPtr<SGraphEditor> GraphEditor, const UE::Slate::FDeprecateVector2DParameter& NewLocation)
 {
 	if (!InterfaceParameter || InterfaceParameter->bHasGraphNode)
 	{
diff --git a/Engine/Plugins/Cameras/GameplayCameras/Source/GameplayCamerasEditor/Private/Editors/CameraNodeGraphSchema.cpp b/Engine/Plugins/Cameras/GameplayCameras/Source/GameplayCamerasEditor/Private/Editors/CameraNodeGraphSchema.cpp
--- a/Engine/Plugins/Cameras/GameplayCameras/Source/GameplayCamerasEditor/Private/Editors/CameraNodeGraphSchema.cpp
+++ b/Engine/Plugins/Cameras/GameplayCameras/Source/GameplayCamerasEditor/Private/Editors/CameraNodeGraphSchema.cpp
@@ -512,7 +512,7 @@ FCameraNodeGraphSchemaAction_NewInterfaceParameterNode::FCameraNodeGraphSchemaAc
 {
 }
 
-UEdGraphNode* FCameraNodeGraphSchemaAction_NewInterfaceParameterNode::PerformAction(UEdGraph* ParentGraph, UEdGraphPin* FromPin, const FVector2D Location, bool bSelectNewNode)
+UEdGraphNode* FCameraNodeGraphSchemaAction_NewInterfaceParameterNode::PerformAction(UEdGraph* ParentGraph, UEdGraphPin* FromPin, const FVector2f& Location, bool bSelectNewNode)
 {
 	UObjectTreeGraph* ObjectTreeGraph = Cast<UObjectTreeGraph>(ParentGraph);
 	if (!ensure(ObjectTreeGraph))
@@ -585,7 +585,7 @@ FCameraNodeGraphSchemaAction_AddInterfaceParameterNode::FCameraNodeGraphSchemaAc
 {
 }
 
-UEdGraphNode* FCameraNodeGraphSchemaAction_AddInterfaceParameterNode::PerformAction(UEdGraph* ParentGraph, UEdGraphPin* FromPin, const FVector2D Location, bool bSelectNewNode)
+UEdGraphNode* FCameraNodeGraphSchemaAction_AddInterfaceParameterNode::PerformAction(UEdGraph* ParentGraph, UEdGraphPin* FromPin, const FVector2f& Location, bool bSelectNewNode)
 {
 	if (!InterfaceParameter || InterfaceParameter->bHasGraphNode)
 	{
diff --git a/Engine/Plugins/Cameras/GameplayCameras/Source/GameplayCamerasEditor/Private/Editors/ObjectTreeDragDropOp.cpp b/Engine/Plugins/Cameras/GameplayCameras/Source/GameplayCamerasEditor/Private/Editors/ObjectTreeDragDropOp.cpp
--- a/Engine/Plugins/Cameras/GameplayCameras/Source/GameplayCamerasEditor/Private/Editors/ObjectTreeDragDropOp.cpp
+++ b/Engine/Plugins/Cameras/GameplayCameras/Source/GameplayCamerasEditor/Private/Editors/ObjectTreeDragDropOp.cpp
@@ -60,7 +60,7 @@ FReply FObjectTreeClassDragDropOp::ExecuteDragOver(TSharedPtr<SGraphEditor> Grap
 	return FReply::Handled();
 }
 
-FReply FObjectTreeClassDragDropOp::ExecuteDrop(TSharedPtr<SGraphEditor> GraphEditor, const FVector2D& NewLocation)
+FReply FObjectTreeClassDragDropOp::ExecuteDrop(TSharedPtr<SGraphEditor> GraphEditor, const UE::Slate::FDeprecateVector2DParameter& NewLocation)
 {
 	const FScopedTransaction Transaction(LOCTEXT("DropObjectClasses", "Drop New Nodes"));
 
@@ -69,7 +69,7 @@ FReply FObjectTreeClassDragDropOp::ExecuteDrop(TSharedPtr<SGraphEditor> GraphEdi
 
 	GraphEditor->ClearSelectionSet();
 
-	FVector2D CurLocation = NewLocation;
+	FVector2f CurLocation = NewLocation;
 	for (UClass* PlaceableClass : PlaceableClasses)
 	{
 		FObjectTreeGraphSchemaAction_NewNode Action;
@@ -77,7 +77,7 @@ FReply FObjectTreeClassDragDropOp::ExecuteDrop(TSharedPtr<SGraphEditor> GraphEdi
 		UEdGraphNode* NewNode = Action.PerformAction(Graph, nullptr, CurLocation, false);
 		GraphEditor->SetNodeSelection(NewNode, true);
 
-		CurLocation += FVector2D(20, 20);
+		CurLocation += FVector2f(20.0f, 20.0f);
 	}
 
 	return FReply::Handled();
diff --git a/Engine/Plugins/Cameras/GameplayCameras/Source/GameplayCamerasEditor/Private/Editors/ObjectTreeGraphSchema.cpp b/Engine/Plugins/Cameras/GameplayCameras/Source/GameplayCamerasEditor/Private/Editors/ObjectTreeGraphSchema.cpp
--- a/Engine/Plugins/Cameras/GameplayCameras/Source/GameplayCamerasEditor/Private/Editors/ObjectTreeGraphSchema.cpp
+++ b/Engine/Plugins/Cameras/GameplayCameras/Source/GameplayCamerasEditor/Private/Editors/ObjectTreeGraphSchema.cpp
@@ -1457,7 +1457,7 @@ FObjectTreeGraphSchemaAction_NewNode::FObjectTreeGraphSchemaAction_NewNode(FText
 {
 }
 
-UEdGraphNode* FObjectTreeGraphSchemaAction_NewNode::PerformAction(UEdGraph* ParentGraph, UEdGraphPin* FromPin, const FVector2D Location, bool bSelectNewNode)
+UEdGraphNode* FObjectTreeGraphSchemaAction_NewNode::PerformAction(UEdGraph* ParentGraph, UEdGraphPin* FromPin, const FVector2f& Location, bool bSelectNewNode)
 {
 	UObjectTreeGraph* ObjectTreeGraph = Cast<UObjectTreeGraph>(ParentGraph);
 	if (!ensure(ObjectTreeGraph))
@@ -1535,7 +1535,7 @@ FObjectTreeGraphSchemaAction_NewComment::FObjectTreeGraphSchemaAction_NewComment
 {
 }
 
-UEdGraphNode* FObjectTreeGraphSchemaAction_NewComment::PerformAction(UEdGraph* ParentGraph, UEdGraphPin* FromPin, const FVector2D Location, bool bSelectNewNode)
+UEdGraphNode* FObjectTreeGraphSchemaAction_NewComment::PerformAction(UEdGraph* ParentGraph, UEdGraphPin* FromPin, const FVector2f& Location, bool bSelectNewNode)
 {
 	UObjectTreeGraph* ObjectTreeGraph = Cast<UObjectTreeGraph>(ParentGraph);
 	if (!ensure(ObjectTreeGraph))
diff --git a/Engine/Plugins/Cameras/GameplayCameras/Source/GameplayCamerasEditor/Private/Editors/SObjectTreeGraphEditor.cpp b/Engine/Plugins/Cameras/GameplayCameras/Source/GameplayCamerasEditor/Private/Editors/SObjectTreeGraphEditor.cpp
--- a/Engine/Plugins/Cameras/GameplayCameras/Source/GameplayCamerasEditor/Private/Editors/SObjectTreeGraphEditor.cpp
+++ b/Engine/Plugins/Cameras/GameplayCameras/Source/GameplayCamerasEditor/Private/Editors/SObjectTreeGraphEditor.cpp
@@ -491,7 +491,7 @@ void SObjectTreeGraphEditor::PasteNodes()
 	FString TextToImport;
 	FPlatformApplicationMisc::ClipboardPaste(TextToImport);
 	
-	ImportNodesFromText(GraphEditor->GetPasteLocation(), TextToImport);
+	ImportNodesFromText(GraphEditor->GetPasteLocation2f(), TextToImport);
 }
 
 bool SObjectTreeGraphEditor::CanPasteNodes()
diff --git a/Engine/Plugins/Cameras/GameplayCameras/Source/GameplayCamerasEditor/Public/Editors/CameraNodeGraphDragDropOp.h b/Engine/Plugins/Cameras/GameplayCameras/Source/GameplayCamerasEditor/Public/Editors/CameraNodeGraphDragDropOp.h
--- a/Engine/Plugins/Cameras/GameplayCameras/Source/GameplayCamerasEditor/Public/Editors/CameraNodeGraphDragDropOp.h
+++ b/Engine/Plugins/Cameras/GameplayCameras/Source/GameplayCamerasEditor/Public/Editors/CameraNodeGraphDragDropOp.h
@@ -16,7 +16,7 @@ public:
 	static TSharedRef<FCameraNodeGraphInterfaceParameterDragDropOp> New(UCameraObjectInterfaceParameterBase* InInterfaceParameter);
 
 	FReply ExecuteDragOver(TSharedPtr<SGraphEditor> GraphEditor);
-	FReply ExecuteDrop(TSharedPtr<SGraphEditor> GraphEditor, const FVector2D& NewLocation);
+	FReply ExecuteDrop(TSharedPtr<SGraphEditor> GraphEditor, const UE::Slate::FDeprecateVector2DParameter& NewLocation);
 
 private:
 
diff --git a/Engine/Plugins/Cameras/GameplayCameras/Source/GameplayCamerasEditor/Public/Editors/CameraNodeGraphSchema.h b/Engine/Plugins/Cameras/GameplayCameras/Source/GameplayCamerasEditor/Public/Editors/CameraNodeGraphSchema.h
--- a/Engine/Plugins/Cameras/GameplayCameras/Source/GameplayCamerasEditor/Public/Editors/CameraNodeGraphSchema.h
+++ b/Engine/Plugins/Cameras/GameplayCameras/Source/GameplayCamerasEditor/Public/Editors/CameraNodeGraphSchema.h
@@ -82,8 +82,9 @@ public:
 
 	// FEdGraphSchemaAction interface.
 	static FName StaticGetTypeId() { static FName Type("FCameraNodeGraphSchemaAction_NewInterfaceParameterNode"); return Type; }
-	virtual FName GetTypeId() const override { return StaticGetTypeId(); } 
-	virtual UEdGraphNode* PerformAction(UEdGraph* ParentGraph, UEdGraphPin* FromPin, const FVector2D Location, bool bSelectNewNode = true) override;
+	virtual FName GetTypeId() const override { return StaticGetTypeId(); }
+	using FEdGraphSchemaAction::PerformAction; // Prevent hiding of deprecated base class function with FVector2D
+	virtual UEdGraphNode* PerformAction(UEdGraph* ParentGraph, UEdGraphPin* FromPin, const FVector2f& Location, bool bSelectNewNode = true) override;
 };
 
 USTRUCT()
@@ -103,7 +104,8 @@ public:
 
 	// FEdGraphSchemaAction interface.
 	static FName StaticGetTypeId() { static FName Type("FCameraNodeGraphSchemaAction_AddInterfaceParameterNode"); return Type; }
-	virtual FName GetTypeId() const override { return StaticGetTypeId(); } 
-	virtual UEdGraphNode* PerformAction(UEdGraph* ParentGraph, UEdGraphPin* FromPin, const FVector2D Location, bool bSelectNewNode = true) override;
+	virtual FName GetTypeId() const override { return StaticGetTypeId(); }
+	using FEdGraphSchemaAction::PerformAction; // Prevent hiding of deprecated base class function with FVector2D
+	virtual UEdGraphNode* PerformAction(UEdGraph* ParentGraph, UEdGraphPin* FromPin, const FVector2f& Location, bool bSelectNewNode = true) override;
 };
 
diff --git a/Engine/Plugins/Cameras/GameplayCameras/Source/GameplayCamerasEditor/Public/Editors/ObjectTreeDragDropOp.h b/Engine/Plugins/Cameras/GameplayCameras/Source/GameplayCamerasEditor/Public/Editors/ObjectTreeDragDropOp.h
--- a/Engine/Plugins/Cameras/GameplayCameras/Source/GameplayCamerasEditor/Public/Editors/ObjectTreeDragDropOp.h
+++ b/Engine/Plugins/Cameras/GameplayCameras/Source/GameplayCamerasEditor/Public/Editors/ObjectTreeDragDropOp.h
@@ -23,7 +23,7 @@ public:
 	TArrayView<UClass* const> GetObjectClasses() const { return ObjectClasses; }
 
 	FReply ExecuteDragOver(TSharedPtr<SGraphEditor> GraphEditor);
-	FReply ExecuteDrop(TSharedPtr<SGraphEditor> GraphEditor, const FVector2D& NewLocation);
+	FReply ExecuteDrop(TSharedPtr<SGraphEditor> GraphEditor, const UE::Slate::FDeprecateVector2DParameter& NewLocation);
 
 private:
 
diff --git a/Engine/Plugins/Cameras/GameplayCameras/Source/GameplayCamerasEditor/Public/Editors/ObjectTreeGraphSchema.h b/Engine/Plugins/Cameras/GameplayCameras/Source/GameplayCamerasEditor/Public/Editors/ObjectTreeGraphSchema.h
--- a/Engine/Plugins/Cameras/GameplayCameras/Source/GameplayCamerasEditor/Public/Editors/ObjectTreeGraphSchema.h
+++ b/Engine/Plugins/Cameras/GameplayCameras/Source/GameplayCamerasEditor/Public/Editors/ObjectTreeGraphSchema.h
@@ -164,8 +164,9 @@ public:
 
 	// FEdGraphSchemaAction interface.
 	static FName StaticGetTypeId() { static FName Type("FObjectTreeGraphSchemaAction_NewNode"); return Type; }
-	virtual FName GetTypeId() const override { return StaticGetTypeId(); } 
-	virtual UEdGraphNode* PerformAction(UEdGraph* ParentGraph, UEdGraphPin* FromPin, const FVector2D Location, bool bSelectNewNode = true) override;
+	virtual FName GetTypeId() const override { return StaticGetTypeId(); }
+	using FEdGraphSchemaAction::PerformAction; // Prevent hiding of deprecated base class function with FVector2D
+	virtual UEdGraphNode* PerformAction(UEdGraph* ParentGraph, UEdGraphPin* FromPin, const FVector2f& Location, bool bSelectNewNode = true) override;
 
 protected:
 
@@ -188,6 +189,6 @@ public:
 	// FEdGraphSchemaAction interface.
 	static FName StaticGetTypeId() { static FName Type("FObjectTreeGraphSchemaAction_NewComment"); return Type; }
 	virtual FName GetTypeId() const override { return StaticGetTypeId(); } 
-	virtual UEdGraphNode* PerformAction(UEdGraph* ParentGraph, UEdGraphPin* FromPin, const FVector2D Location, bool bSelectNewNode = true) override;
+	virtual UEdGraphNode* PerformAction(UEdGraph* ParentGraph, UEdGraphPin* FromPin, const FVector2f& Location, bool bSelectNewNode = true) override;
 };
 
diff --git a/Engine/Plugins/ChaosClothAssetEditor/Source/ChaosClothAssetEditor/Private/ChaosClothAsset/ClothEditorMode.cpp b/Engine/Plugins/ChaosClothAssetEditor/Source/ChaosClothAssetEditor/Private/ChaosClothAsset/ClothEditorMode.cpp
--- a/Engine/Plugins/ChaosClothAssetEditor/Source/ChaosClothAssetEditor/Private/ChaosClothAsset/ClothEditorMode.cpp
+++ b/Engine/Plugins/ChaosClothAssetEditor/Source/ChaosClothAssetEditor/Private/ChaosClothAsset/ClothEditorMode.cpp
@@ -1697,7 +1697,7 @@ UEdGraphNode* UChaosClothAssetEditorMode::CreateNewNode(const FName& NewNodeType
 		FAssetSchemaAction_Dataflow_CreateNode_DataflowEdNode::CreateAction(DataflowGraph.Get(), NewNodeTypeName);
 	constexpr UEdGraphPin* FromPin = nullptr;
 	constexpr bool bSelectNewNode = true;
-	UEdGraphNode* const NewEdNode = NodeAction->PerformAction(DataflowGraph.Get(), FromPin, PinnedDataflowGraphEditor->GetPasteLocation(), bSelectNewNode);
+	UEdGraphNode* const NewEdNode = NodeAction->PerformAction(DataflowGraph.Get(), FromPin, PinnedDataflowGraphEditor->GetPasteLocation2f(), bSelectNewNode);
 
 	return NewEdNode;
 }
diff --git a/Engine/Plugins/Editor/AssetManagerEditor/Source/AssetManagerEditor/Private/ReferenceViewer/ReferenceViewerSchema.cpp b/Engine/Plugins/Editor/AssetManagerEditor/Source/AssetManagerEditor/Private/ReferenceViewer/ReferenceViewerSchema.cpp
--- a/Engine/Plugins/Editor/AssetManagerEditor/Source/AssetManagerEditor/Private/ReferenceViewer/ReferenceViewerSchema.cpp
+++ b/Engine/Plugins/Editor/AssetManagerEditor/Source/AssetManagerEditor/Private/ReferenceViewer/ReferenceViewerSchema.cpp
@@ -221,7 +221,7 @@ FConnectionDrawingPolicy* UReferenceViewerSchema::CreateConnectionDrawingPolicy(
 	return new FReferenceViewerConnectionDrawingPolicy(InBackLayerID, InFrontLayerID, InZoomFactor, InClippingRect, InDrawElements);
 }
 
-void UReferenceViewerSchema::DroppedAssetsOnGraph(const TArray<FAssetData>& Assets, const FVector2D& GraphPosition, UEdGraph* Graph) const
+void UReferenceViewerSchema::DroppedAssetsOnGraph(const TArray<FAssetData>& Assets, const FVector2f& GraphPosition, UEdGraph* Graph) const
 {
 	TArray<FAssetIdentifier> AssetIdentifiers;
 
diff --git a/Engine/Plugins/Editor/AssetManagerEditor/Source/AssetManagerEditor/Public/ReferenceViewer/ReferenceViewerSchema.h b/Engine/Plugins/Editor/AssetManagerEditor/Source/AssetManagerEditor/Public/ReferenceViewer/ReferenceViewerSchema.h
--- a/Engine/Plugins/Editor/AssetManagerEditor/Source/AssetManagerEditor/Public/ReferenceViewer/ReferenceViewerSchema.h
+++ b/Engine/Plugins/Editor/AssetManagerEditor/Source/AssetManagerEditor/Public/ReferenceViewer/ReferenceViewerSchema.h
@@ -27,7 +27,7 @@ public:
 	virtual FPinConnectionResponse MovePinLinks(UEdGraphPin& MoveFromPin, UEdGraphPin& MoveToPin, bool bIsIntermediateMove = false, bool bNotifyLinkedNodes = false) const override;
 	virtual FPinConnectionResponse CopyPinLinks(UEdGraphPin& CopyFromPin, UEdGraphPin& CopyToPin, bool bIsIntermediateCopy = false) const override;
 	virtual class FConnectionDrawingPolicy* CreateConnectionDrawingPolicy(int32 InBackLayerID, int32 InFrontLayerID, float InZoomFactor, const FSlateRect& InClippingRect, class FSlateWindowElementList& InDrawElements, class UEdGraph* InGraphObj) const override;
-	virtual void DroppedAssetsOnGraph(const TArray<FAssetData>& Assets, const FVector2D& GraphPosition, UEdGraph* Graph) const override;
+	virtual void DroppedAssetsOnGraph(const TArray<FAssetData>& Assets, const FVector2f& GraphPosition, UEdGraph* Graph) const override;
 	virtual void GetAssetsGraphHoverMessage(const TArray<FAssetData>& Assets, const UEdGraph* HoverGraph, FString& OutTooltipText, bool& OutOkIcon) const override;
 	// End of UEdGraphSchema interface
 
diff --git a/Engine/Plugins/Enterprise/DataprepEditor/Source/DataprepEditor/Private/DataprepGraph/DataprepGraphSchema.h b/Engine/Plugins/Enterprise/DataprepEditor/Source/DataprepEditor/Private/DataprepGraph/DataprepGraphSchema.h
--- a/Engine/Plugins/Enterprise/DataprepEditor/Source/DataprepEditor/Private/DataprepGraph/DataprepGraphSchema.h
+++ b/Engine/Plugins/Enterprise/DataprepEditor/Source/DataprepEditor/Private/DataprepGraph/DataprepGraphSchema.h
@@ -16,7 +16,7 @@ public:
 	// UEdGraphSchema interface
 	virtual void GetContextMenuActions(class UToolMenu* Menu, class UGraphNodeContextMenuContext* Context) const override {}
 	virtual FName GetParentContextMenuName() const override { return NAME_None; }
-	virtual void DroppedAssetsOnGraph(const TArray<FAssetData>& Assets, const FVector2D& GraphPosition, UEdGraph* Graph) const override {}
+	virtual void DroppedAssetsOnGraph(const TArray<FAssetData>& Assets, const FVector2f& GraphPosition, UEdGraph* Graph) const override {}
 	virtual void GetAssetsGraphHoverMessage(const TArray<FAssetData>& Assets, const UEdGraph* HoverGraph, FString& OutTooltipText, bool& OutOkIcon) const override {}
 	//virtual class FConnectionDrawingPolicy* CreateConnectionDrawingPolicy(int32 InBackLayerID, int32 InFrontLayerID, float InZoomFactor, const FSlateRect& InClippingRect, class FSlateWindowElementList& InDrawElements, class UEdGraph* InGraphObj) const override;
 	// End of UEdGraphSchema interface
diff --git a/Engine/Plugins/Enterprise/DataprepEditor/Source/DataprepEditor/Private/Widgets/DataprepGraph/SDataprepGraphEditor.cpp b/Engine/Plugins/Enterprise/DataprepEditor/Source/DataprepEditor/Private/Widgets/DataprepGraph/SDataprepGraphEditor.cpp
--- a/Engine/Plugins/Enterprise/DataprepEditor/Source/DataprepEditor/Private/Widgets/DataprepGraph/SDataprepGraphEditor.cpp
+++ b/Engine/Plugins/Enterprise/DataprepEditor/Source/DataprepEditor/Private/Widgets/DataprepGraph/SDataprepGraphEditor.cpp
@@ -276,17 +276,17 @@ void SDataprepGraphEditor::Tick(const FGeometry& AllottedGeometry, const double
 			}
 		}
 
-		FVector2D Location;
+		FVector2f Location;
 		float ZoomAmount = 1.f;
 		GetViewLocation( Location, ZoomAmount );
 
-		UpdateLayout( AllottedGeometry.GetLocalSize(), Location, ZoomAmount );
+		UpdateLayout( AllottedGeometry.GetLocalSize(), FDeprecateSlateVector2D(Location), ZoomAmount );
 	}
 
 	SGraphEditor::Tick( AllottedGeometry, InCurrentTime, InDeltaTime);
 }
 
-void SDataprepGraphEditor::UpdateLayout( const FVector2D& LocalSize, const FVector2D& Location, float ZoomAmount )
+void SDataprepGraphEditor::UpdateLayout( const UE::Slate::FDeprecateVector2DParameter& LocalSize, const UE::Slate::FDeprecateVector2DParameter& Location, float ZoomAmount )
 {
 	if(SDataprepGraphTrackNode* TrackGraphNode = TrackGraphNodePtr.Pin().Get())
 	{
@@ -302,10 +302,10 @@ void SDataprepGraphEditor::UpdateLayout( const FVector2D& LocalSize, const FVect
 
 		if( !Location.Equals(LastLocation) )
 		{
-			FVector2D ComputedLocation(Location);
+			FVector2f ComputedLocation(Location);
 
-			FVector2D PanelSize = LocalSize / ZoomAmount;
-			FVector2D WorkingSize = WorkingArea.GetSize();
+			FVector2f PanelSize = LocalSize / ZoomAmount;
+			FVector2f WorkingSize = WorkingArea.GetSize();
 
 			if(Location.X != LastLocation.X)
 			{
@@ -1290,7 +1290,7 @@ FActionMenuContent SDataprepGraphEditor::OnCreateNodeOrPinMenu(UEdGraph* Current
 	}
 
 	// Create contextual for graph panel when on top of track node too
-	return OnCreateActionMenu(CurrentGraph, GetPasteLocation(), TArray<UEdGraphPin*>(), true, SGraphEditor::FActionMenuClosed());
+	return OnCreateActionMenu(CurrentGraph, GetPasteLocation2f(), TArray<UEdGraphPin*>(), true, SGraphEditor::FActionMenuClosed());
 }
 
 #undef LOCTEXT_NAMESPACE
diff --git a/Engine/Plugins/Enterprise/DataprepEditor/Source/DataprepEditor/Private/Widgets/DataprepGraph/SDataprepGraphEditor.h b/Engine/Plugins/Enterprise/DataprepEditor/Source/DataprepEditor/Private/Widgets/DataprepGraph/SDataprepGraphEditor.h
--- a/Engine/Plugins/Enterprise/DataprepEditor/Source/DataprepEditor/Private/Widgets/DataprepGraph/SDataprepGraphEditor.h
+++ b/Engine/Plugins/Enterprise/DataprepEditor/Source/DataprepEditor/Private/Widgets/DataprepGraph/SDataprepGraphEditor.h
@@ -132,7 +132,7 @@ public:
 
 private:
 	/** Recompute the layout of the displayed graph after a pan, resize and/or zoom */
-	void UpdateLayout( const FVector2D& LocalSize, const FVector2D& Location, float ZoomAmount );
+	void UpdateLayout( const UE::Slate::FDeprecateVector2DParameter& LocalSize, const UE::Slate::FDeprecateVector2DParameter& Location, float ZoomAmount );
 
 	void BuildCommandList();
 
@@ -141,10 +141,10 @@ private:
 	mutable bool bIsComplete;
 
 	/** Last size of the window displaying the graph's canvas */
-	FVector2D LastLocalSize;
+	FDeprecateSlateVector2D LastLocalSize;
 
 	/** Last location of the upper left corner of the visible section of the graph's canvas */
-	FVector2D LastLocation;
+	FDeprecateSlateVector2D LastLocation;
 
 	/** Last zoom factor applied to the graph's canvas */
 	float LastZoomAmount;
diff --git a/Engine/Plugins/Experimental/AnimNext/Source/AnimNextEditor/Internal/Common/GraphEditorSchemaActions.h b/Engine/Plugins/Experimental/AnimNext/Source/AnimNextEditor/Internal/Common/GraphEditorSchemaActions.h
--- a/Engine/Plugins/Experimental/AnimNext/Source/AnimNextEditor/Internal/Common/GraphEditorSchemaActions.h
+++ b/Engine/Plugins/Experimental/AnimNext/Source/AnimNextEditor/Internal/Common/GraphEditorSchemaActions.h
@@ -45,8 +45,8 @@ struct ANIMNEXTEDITOR_API FAnimNextSchemaAction_RigUnit : public FAnimNextSchema
 	{}
 
 	// FEdGraphSchemaAction Interface
-	virtual UEdGraphNode* PerformAction(UEdGraph* ParentGraph, TArray<UEdGraphPin*>& FromPins, const FVector2D Location, bool bSelectNewNode = true) override;
-	virtual UEdGraphNode* PerformAction(UEdGraph* ParentGraph, UEdGraphPin* FromPin, const FVector2D Location, bool bSelectNewNode = true) { return nullptr; }
+	virtual UEdGraphNode* PerformAction(UEdGraph* ParentGraph, TArray<UEdGraphPin*>& FromPins, const FVector2f& Location, bool bSelectNewNode = true) override;
+	virtual UEdGraphNode* PerformAction(UEdGraph* ParentGraph, UEdGraphPin* FromPin, const FVector2f& Location, bool bSelectNewNode = true) { return nullptr; }
 
 private:
 	// The script struct for our rig unit
@@ -70,8 +70,8 @@ struct FAnimNextSchemaAction_DispatchFactory : public FAnimNextSchemaAction
 	virtual const FSlateBrush* GetIconBrush() const override;
 	
 	// FEdGraphSchemaAction Interface
-	virtual UEdGraphNode* PerformAction(UEdGraph* ParentGraph, TArray<UEdGraphPin*>& FromPins, const FVector2D Location, bool bSelectNewNode = true) override;
-	virtual UEdGraphNode* PerformAction(UEdGraph* ParentGraph, UEdGraphPin* FromPin, const FVector2D Location, bool bSelectNewNode = true) { return nullptr; }
+	virtual UEdGraphNode* PerformAction(UEdGraph* ParentGraph, TArray<UEdGraphPin*>& FromPins, const FVector2f& Location, bool bSelectNewNode = true) override;
+	virtual UEdGraphNode* PerformAction(UEdGraph* ParentGraph, UEdGraphPin* FromPin, const FVector2f& Location, bool bSelectNewNode = true) { return nullptr; }
 
 private:
 	// Notation for dispatch factory
@@ -96,8 +96,8 @@ struct FAnimNextSchemaAction_Variable : public FAnimNextSchemaAction
 	FAnimNextSchemaAction_Variable(FName InName, const FAnimNextParamType& InType, const EVariableAccessorChoice InVariableAccessorChoice);
 
 	// FEdGraphSchemaAction Interface
-	virtual UEdGraphNode* PerformAction(UEdGraph* ParentGraph, TArray<UEdGraphPin*>& FromPins, const FVector2D Location, bool bSelectNewNode = true) override;
-	virtual UEdGraphNode* PerformAction(UEdGraph* ParentGraph, UEdGraphPin* FromPin, const FVector2D Location, bool bSelectNewNode = true) { return nullptr; }
+	virtual UEdGraphNode* PerformAction(UEdGraph* ParentGraph, TArray<UEdGraphPin*>& FromPins, const FVector2f& Location, bool bSelectNewNode = true) override;
+	virtual UEdGraphNode* PerformAction(UEdGraph* ParentGraph, UEdGraphPin* FromPin, const FVector2f& Location, bool bSelectNewNode = true) { return nullptr; }
 
 	virtual const FSlateBrush* GetIconBrush() const override;
 
@@ -119,8 +119,8 @@ struct FAnimNextSchemaAction_AddComment : public FAnimNextSchemaAction
 	FAnimNextSchemaAction_AddComment();
 
 	// FEdGraphSchemaAction Interface
-	virtual UEdGraphNode* PerformAction(UEdGraph* ParentGraph, TArray<UEdGraphPin*>& FromPins, const FVector2D Location, bool bSelectNewNode = true) override;
-	virtual UEdGraphNode* PerformAction(UEdGraph* ParentGraph, UEdGraphPin* FromPin, const FVector2D Location, bool bSelectNewNode = true) { return nullptr; }
+	virtual UEdGraphNode* PerformAction(UEdGraph* ParentGraph, TArray<UEdGraphPin*>& FromPins, const FVector2f& Location, bool bSelectNewNode = true) override;
+	virtual UEdGraphNode* PerformAction(UEdGraph* ParentGraph, UEdGraphPin* FromPin, const FVector2f& Location, bool bSelectNewNode = true) { return nullptr; }
 
 	virtual const FSlateBrush* GetIconBrush() const override;
 };
@@ -138,8 +138,8 @@ struct FAnimNextSchemaAction_Function : public FAnimNextSchemaAction
 	virtual const FSlateBrush* GetIconBrush() const override;
 	
 	// FEdGraphSchemaAction Interface
-	virtual UEdGraphNode* PerformAction(UEdGraph* ParentGraph, TArray<UEdGraphPin*>& FromPins, const FVector2D Location, bool bSelectNewNode = true) override;
-	virtual UEdGraphNode* PerformAction(UEdGraph* ParentGraph, UEdGraphPin* FromPin, const FVector2D Location, bool bSelectNewNode = true) { return nullptr; }
+	virtual UEdGraphNode* PerformAction(UEdGraph* ParentGraph, TArray<UEdGraphPin*>& FromPins, const FVector2f& Location, bool bSelectNewNode = true) override;
+	virtual UEdGraphNode* PerformAction(UEdGraph* ParentGraph, UEdGraphPin* FromPin, const FVector2f& Location, bool bSelectNewNode = true) { return nullptr; }
 
 private:
 	/** The public function definition we will spawn from [optional] */
diff --git a/Engine/Plugins/Experimental/AnimNext/Source/AnimNextEditor/Internal/Common/SActionMenu.h b/Engine/Plugins/Experimental/AnimNext/Source/AnimNextEditor/Internal/Common/SActionMenu.h
--- a/Engine/Plugins/Experimental/AnimNext/Source/AnimNextEditor/Internal/Common/SActionMenu.h
+++ b/Engine/Plugins/Experimental/AnimNext/Source/AnimNextEditor/Internal/Common/SActionMenu.h
@@ -83,7 +83,7 @@ private:
 	bool bActionExecuted = false;
 
 	TArray<UEdGraphPin*> DraggedFromPins;
-	FVector2D NewNodePosition = FVector2D::ZeroVector;
+	FDeprecateSlateVector2D NewNodePosition = FVector2f::ZeroVector;
 
 	SGraphEditor::FActionMenuClosed OnClosedCallback;
 	FClosedReason OnCloseReasonCallback;
diff --git a/Engine/Plugins/Experimental/AnimNext/Source/AnimNextEditor/Private/Common/GraphEditorSchemaActions.cpp b/Engine/Plugins/Experimental/AnimNext/Source/AnimNextEditor/Private/Common/GraphEditorSchemaActions.cpp
--- a/Engine/Plugins/Experimental/AnimNext/Source/AnimNextEditor/Private/Common/GraphEditorSchemaActions.cpp
+++ b/Engine/Plugins/Experimental/AnimNext/Source/AnimNextEditor/Private/Common/GraphEditorSchemaActions.cpp
@@ -32,7 +32,7 @@ const FLinearColor& FAnimNextSchemaAction::GetIconColor() const
 
 // *** Rig Unit ***
 
-UEdGraphNode* FAnimNextSchemaAction_RigUnit::PerformAction(UEdGraph* ParentGraph, TArray<UEdGraphPin*>& FromPins, const FVector2D Location, bool bSelectNewNode)
+UEdGraphNode* FAnimNextSchemaAction_RigUnit::PerformAction(UEdGraph* ParentGraph, TArray<UEdGraphPin*>& FromPins, const FVector2f& Location, bool bSelectNewNode)
 {
 	IRigVMClientHost* Host = ParentGraph->GetImplementingOuter<IRigVMClientHost>();
 	URigVMEdGraphNode* NewNode = nullptr;
@@ -54,7 +54,7 @@ UEdGraphNode* FAnimNextSchemaAction_RigUnit::PerformAction(UEdGraph* ParentGraph
 		FRigVMUnitNodeCreatedContext& UnitNodeCreatedContext = Controller->GetUnitNodeCreatedContext();
 		FRigVMUnitNodeCreatedContext::FScope ReasonScope(UnitNodeCreatedContext, ERigVMNodeCreatedReason::NodeSpawner, Host);
 
-		if (URigVMUnitNode* ModelNode = Controller->AddUnitNode(StructTemplate, NodeClass, FRigVMStruct::ExecuteName, Location, Name.ToString(), true, true))
+		if (URigVMUnitNode* ModelNode = Controller->AddUnitNode(StructTemplate, NodeClass, FRigVMStruct::ExecuteName, FDeprecateSlateVector2D(Location), Name.ToString(), true, true))
 		{
 			NewNode = Cast<URigVMEdGraphNode>(EdGraph->FindNodeForModelNodeName(ModelNode->GetFName()));
 			check(NewNode);
@@ -83,7 +83,7 @@ const FSlateBrush* FAnimNextSchemaAction_DispatchFactory::GetIconBrush() const
 	return FRigVMEditorStyle::Get().GetBrush("RigVM.Template");
 }
 
-UEdGraphNode* FAnimNextSchemaAction_DispatchFactory::PerformAction(UEdGraph* ParentGraph, TArray<UEdGraphPin*>& FromPins, const FVector2D Location, bool bSelectNewNode)
+UEdGraphNode* FAnimNextSchemaAction_DispatchFactory::PerformAction(UEdGraph* ParentGraph, TArray<UEdGraphPin*>& FromPins, const FVector2f& Location, bool bSelectNewNode)
 {
 	IRigVMClientHost* Host = ParentGraph->GetImplementingOuter<IRigVMClientHost>();
 	if(Host == nullptr)
@@ -118,7 +118,7 @@ UEdGraphNode* FAnimNextSchemaAction_DispatchFactory::PerformAction(UEdGraph* Par
 
 	Controller->OpenUndoBracket(FString::Printf(TEXT("Add '%s' Node"), *Name.ToString()));
 
-	if (URigVMTemplateNode* ModelNode = Controller->AddTemplateNode(Notation, Location, Name.ToString(), true, true))
+	if (URigVMTemplateNode* ModelNode = Controller->AddTemplateNode(Notation, FDeprecateSlateVector2D(Location), Name.ToString(), true, true))
 	{
 		NewNode = Cast<URigVMEdGraphNode>(EdGraph->FindNodeForModelNodeName(ModelNode->GetFName()));
 
@@ -194,9 +194,9 @@ const FLinearColor& FAnimNextSchemaAction_Variable::GetIconColor() const
 	return VariableColor;
 }
 
-UEdGraphNode* FAnimNextSchemaAction_Variable::PerformAction(UEdGraph* ParentGraph, TArray<UEdGraphPin*>& FromPins, const FVector2D Location, bool bSelectNewNode)
+UEdGraphNode* FAnimNextSchemaAction_Variable::PerformAction(UEdGraph* ParentGraph, TArray<UEdGraphPin*>& FromPins, const FVector2f& Location, bool bSelectNewNode)
 {
-	const auto AddNodeLambda = [](UEdGraph* ParentGraph, FName InName, const FString& InTypeName, const FString& InObjectPath, const FVector2D Location, const bool bIsGetter) -> URigVMEdGraphNode*
+	const auto AddNodeLambda = [](UEdGraph* ParentGraph, FName InName, const FString& InTypeName, const FString& InObjectPath, const FVector2f& Location, const bool bIsGetter) -> URigVMEdGraphNode*
 	{
 		IRigVMClientHost* Host = ParentGraph->GetImplementingOuter<IRigVMClientHost>();
 		if(Host == nullptr)
@@ -217,7 +217,7 @@ UEdGraphNode* FAnimNextSchemaAction_Variable::PerformAction(UEdGraph* ParentGrap
 		URigVMController* Controller = Host->GetRigVMClient()->GetController(ParentGraph);
 		Controller->OpenUndoBracket(TEXT("Add Variable"));
 
-		if (URigVMNode* ModelNode = Controller->AddVariableNodeFromObjectPath(InName, InTypeName, InObjectPath, bIsGetter, FString(), Location, NodeName, true, true))
+		if (URigVMNode* ModelNode = Controller->AddVariableNodeFromObjectPath(InName, InTypeName, InObjectPath, bIsGetter, FString(), FDeprecateSlateVector2D(Location), NodeName, true, true))
 		{
 			for (UEdGraphNode* Node : ParentGraph->Nodes)
 			{
@@ -305,11 +305,11 @@ FAnimNextSchemaAction_AddComment::FAnimNextSchemaAction_AddComment()
 {
 }
 
-UEdGraphNode* FAnimNextSchemaAction_AddComment::PerformAction(UEdGraph* ParentGraph, TArray<UEdGraphPin*>& FromPins, const FVector2D Location, bool bSelectNewNode/* = true*/)
+UEdGraphNode* FAnimNextSchemaAction_AddComment::PerformAction(UEdGraph* ParentGraph, TArray<UEdGraphPin*>& FromPins, const FVector2f& Location, bool bSelectNewNode/* = true*/)
 {
 	UEdGraphNode_Comment* const CommentTemplate = NewObject<UEdGraphNode_Comment>();
 
-	FVector2D SpawnLocation = Location;
+	FVector2f SpawnLocation = Location;
 	FSlateRect Bounds;
 
 	TSharedPtr<SGraphEditor> GraphEditorPtr = SGraphEditor::FindGraphEditorForGraph(ParentGraph);
@@ -351,7 +351,7 @@ const FSlateBrush* FAnimNextSchemaAction_Function::GetIconBrush() const
 	return FAppStyle::GetBrush("GraphEditor.Function_16x");
 }
 
-UEdGraphNode* FAnimNextSchemaAction_Function::PerformAction(UEdGraph* ParentGraph, TArray<UEdGraphPin*>& FromPins, const FVector2D Location, bool bSelectNewNode)
+UEdGraphNode* FAnimNextSchemaAction_Function::PerformAction(UEdGraph* ParentGraph, TArray<UEdGraphPin*>& FromPins, const FVector2f& Location, bool bSelectNewNode)
 {
 	IRigVMClientHost* Host = ParentGraph->GetImplementingOuter<IRigVMClientHost>();
 	URigVMEdGraphNode* NewNode = nullptr;
@@ -370,7 +370,7 @@ UEdGraphNode* FAnimNextSchemaAction_Function::PerformAction(UEdGraph* ParentGrap
 
 		Controller->OpenUndoBracket(FString::Printf(TEXT("Add '%s' Node"), *Name.ToString()));
 
-		if (URigVMFunctionReferenceNode* ModelNode = Controller->AddFunctionReferenceNodeFromDescription(ReferencedPublicFunctionHeader, Location, Name.ToString(), true, true))
+		if (URigVMFunctionReferenceNode* ModelNode = Controller->AddFunctionReferenceNodeFromDescription(ReferencedPublicFunctionHeader, FDeprecateSlateVector2D(Location), Name.ToString(), true, true))
 		{
 			NewNode = Cast<URigVMEdGraphNode>(EdGraph->FindNodeForModelNodeName(ModelNode->GetFName()));
 			check(NewNode);
diff --git a/Engine/Plugins/Experimental/AnimNextAnimGraph/Source/AnimNextAnimGraphUncookedOnly/Internal/Graph/AnimGraphEditorSchemaActions.cpp b/Engine/Plugins/Experimental/AnimNextAnimGraph/Source/AnimNextAnimGraphUncookedOnly/Internal/Graph/AnimGraphEditorSchemaActions.cpp
--- a/Engine/Plugins/Experimental/AnimNextAnimGraph/Source/AnimNextAnimGraphUncookedOnly/Internal/Graph/AnimGraphEditorSchemaActions.cpp
+++ b/Engine/Plugins/Experimental/AnimNextAnimGraph/Source/AnimNextAnimGraphUncookedOnly/Internal/Graph/AnimGraphEditorSchemaActions.cpp
@@ -41,7 +41,7 @@ FAnimNextSchemaAction_AddManifestNode::FAnimNextSchemaAction_AddManifestNode(con
 {
 }
 
-UEdGraphNode* FAnimNextSchemaAction_AddManifestNode::PerformAction(UEdGraph* ParentGraph, TArray<UEdGraphPin*>& FromPins, const FVector2D Location, bool bSelectNewNode)
+UEdGraphNode* FAnimNextSchemaAction_AddManifestNode::PerformAction(UEdGraph* ParentGraph, TArray<UEdGraphPin*>& FromPins, const FVector2f& Location, bool bSelectNewNode)
 {
 	IRigVMClientHost* Host = ParentGraph->GetImplementingOuter<IRigVMClientHost>();
 	URigVMEdGraphNode* NewNode = nullptr;
@@ -72,7 +72,7 @@ UEdGraphNode* FAnimNextSchemaAction_AddManifestNode::PerformAction(UEdGraph* Par
 
 			const FString NodeData = MoveTemp(Archive);
 
-			TArray<FName> NodeNamesCreated = UE::RigVM::Editor::Tools::ImportNodesFromText(Location, NodeData, Controller, EdGraph->GetModel(), Host->GetLocalFunctionLibrary(), Host->GetRigVMGraphFunctionHost(), true, true);
+			TArray<FName> NodeNamesCreated = UE::RigVM::Editor::Tools::ImportNodesFromText(FDeprecateSlateVector2D(Location), NodeData, Controller, EdGraph->GetModel(), Host->GetLocalFunctionLibrary(), Host->GetRigVMGraphFunctionHost(), true, true);
 
 			TArray<URigVMNode*> ModelNodes;
 			for (const FName& Name : NodeNamesCreated)
@@ -110,7 +110,7 @@ FAnimNextSchemaAction_NotifyEvent::FAnimNextSchemaAction_NotifyEvent()
 {
 }
 
-UEdGraphNode* FAnimNextSchemaAction_NotifyEvent::PerformAction(UEdGraph* ParentGraph, TArray<UEdGraphPin*>& FromPins, const FVector2D Location, bool bSelectNewNode)
+UEdGraphNode* FAnimNextSchemaAction_NotifyEvent::PerformAction(UEdGraph* ParentGraph, TArray<UEdGraphPin*>& FromPins, const FVector2f& Location, bool bSelectNewNode)
 {
 	URigVMEdGraph* EdGraph = Cast<URigVMEdGraph>(ParentGraph);
 	URigVMController* Controller = Cast<URigVMController>(EdGraph->GetController());
@@ -127,7 +127,7 @@ UEdGraphNode* FAnimNextSchemaAction_NotifyEvent::PerformAction(UEdGraph* ParentG
 		if(!Controller->GetAllEventNames().Contains(InNotifyEventName))
 		{
 			Controller->OpenUndoBracket(LOCTEXT("AddNotifyEventTransaction", "Add Notify Event").ToString());
-			URigVMUnitNode* CustomEventNode = Controller->AddUnitNode(FRigVMFunction_UserDefinedEvent::StaticStruct(), FRigVMStruct::ExecuteName, Location);
+			URigVMUnitNode* CustomEventNode = Controller->AddUnitNode(FRigVMFunction_UserDefinedEvent::StaticStruct(), FRigVMStruct::ExecuteName, FDeprecateSlateVector2D(Location));
 			Controller->SetPinDefaultValue(CustomEventNode->FindPin(GET_MEMBER_NAME_STRING_CHECKED(FRigVMFunction_UserDefinedEvent, EventName))->GetPinPath(), InNotifyEventName.ToString(), true, true, true, true, true);
 			Controller->CloseUndoBracket();
 		}
diff --git a/Engine/Plugins/Experimental/AnimNextAnimGraph/Source/AnimNextAnimGraphUncookedOnly/Internal/Graph/AnimGraphEditorSchemaActions.h b/Engine/Plugins/Experimental/AnimNextAnimGraph/Source/AnimNextAnimGraphUncookedOnly/Internal/Graph/AnimGraphEditorSchemaActions.h
--- a/Engine/Plugins/Experimental/AnimNextAnimGraph/Source/AnimNextAnimGraphUncookedOnly/Internal/Graph/AnimGraphEditorSchemaActions.h
+++ b/Engine/Plugins/Experimental/AnimNextAnimGraph/Source/AnimNextAnimGraphUncookedOnly/Internal/Graph/AnimGraphEditorSchemaActions.h
@@ -17,8 +17,8 @@ struct ANIMNEXTANIMGRAPHUNCOOKEDONLY_API FAnimNextSchemaAction_AddManifestNode :
 	FAnimNextSchemaAction_AddManifestNode(const FAnimNextAssetRegistryManifestNode& InManifestNodeData, const FText& InKeywords = FText::GetEmpty());
 
 	// FEdGraphSchemaAction Interface
-	virtual UEdGraphNode* PerformAction(UEdGraph* ParentGraph, TArray<UEdGraphPin*>& FromPins, const FVector2D Location, bool bSelectNewNode = true) override;
-	virtual UEdGraphNode* PerformAction(UEdGraph* ParentGraph, UEdGraphPin* FromPin, const FVector2D Location, bool bSelectNewNode = true) override { return nullptr; }
+	virtual UEdGraphNode* PerformAction(UEdGraph* ParentGraph, TArray<UEdGraphPin*>& FromPins, const FVector2f& Location, bool bSelectNewNode = true) override;
+	virtual UEdGraphNode* PerformAction(UEdGraph* ParentGraph, UEdGraphPin* FromPin, const FVector2f& Location, bool bSelectNewNode = true) override { return nullptr; }
 
 	TSoftObjectPtr<URigVMGraph> ModelGraph;
 	FString NodeName;
@@ -32,6 +32,6 @@ struct ANIMNEXTANIMGRAPHUNCOOKEDONLY_API FAnimNextSchemaAction_NotifyEvent : pub
 	FAnimNextSchemaAction_NotifyEvent();
 
 	// FEdGraphSchemaAction Interface
-	virtual UEdGraphNode* PerformAction(UEdGraph* ParentGraph, TArray<UEdGraphPin*>& FromPins, const FVector2D Location, bool bSelectNewNode = true) override;
-	virtual UEdGraphNode* PerformAction(UEdGraph* ParentGraph, UEdGraphPin* FromPin, const FVector2D Location, bool bSelectNewNode = true) override { return nullptr; }
+	virtual UEdGraphNode* PerformAction(UEdGraph* ParentGraph, TArray<UEdGraphPin*>& FromPins, const FVector2f& Location, bool bSelectNewNode = true) override;
+	virtual UEdGraphNode* PerformAction(UEdGraph* ParentGraph, UEdGraphPin* FromPin, const FVector2f& Location, bool bSelectNewNode = true) override { return nullptr; }
 };
diff --git a/Engine/Plugins/Experimental/Avalanche/Source/AvalancheMediaEditor/Private/Playback/AvaPlaybackGraphEditor.cpp b/Engine/Plugins/Experimental/Avalanche/Source/AvalancheMediaEditor/Private/Playback/AvaPlaybackGraphEditor.cpp
--- a/Engine/Plugins/Experimental/Avalanche/Source/AvalancheMediaEditor/Private/Playback/AvaPlaybackGraphEditor.cpp
+++ b/Engine/Plugins/Experimental/Avalanche/Source/AvalancheMediaEditor/Private/Playback/AvaPlaybackGraphEditor.cpp
@@ -514,7 +514,7 @@ void FAvaPlaybackGraphEditor::CreateComment()
 	if (GraphEditor.IsValid())
 	{
 		FAvaPlaybackAction_NewComment CommentAction;
-		CommentAction.PerformAction(GraphEditor->GetCurrentGraph(), nullptr, GraphEditor->GetPasteLocation());
+		CommentAction.PerformAction(GraphEditor->GetCurrentGraph(), nullptr, GraphEditor->GetPasteLocation2f());
 	}
 }
 
@@ -695,7 +695,7 @@ void FAvaPlaybackGraphEditor::PasteNodes()
 {
 	if (GraphEditor.IsValid())
 	{
-		PasteNodesHere(GraphEditor->GetPasteLocation());
+		PasteNodesHere(GraphEditor->GetPasteLocation2f());
 	}
 }
 
diff --git a/Engine/Plugins/Experimental/Avalanche/Source/AvalancheMediaEditor/Private/Playback/Graph/AvaPlaybackEditorGraphSchema.cpp b/Engine/Plugins/Experimental/Avalanche/Source/AvalancheMediaEditor/Private/Playback/Graph/AvaPlaybackEditorGraphSchema.cpp
--- a/Engine/Plugins/Experimental/Avalanche/Source/AvalancheMediaEditor/Private/Playback/Graph/AvaPlaybackEditorGraphSchema.cpp
+++ b/Engine/Plugins/Experimental/Avalanche/Source/AvalancheMediaEditor/Private/Playback/Graph/AvaPlaybackEditorGraphSchema.cpp
@@ -220,13 +220,13 @@ void UAvaPlaybackEditorGraphSchema::GetAssetsGraphHoverMessage(const TArray<FAss
 	Super::GetAssetsGraphHoverMessage(Assets, HoverGraph, OutTooltipText, OutOkIcon);
 }
 
-void UAvaPlaybackEditorGraphSchema::DroppedAssetsOnGraph(const TArray<FAssetData>& Assets, const FVector2D& GraphPosition
+void UAvaPlaybackEditorGraphSchema::DroppedAssetsOnGraph(const TArray<FAssetData>& Assets, const FVector2f& GraphPosition
 	, UEdGraph* Graph) const
 {
 	Super::DroppedAssetsOnGraph(Assets, GraphPosition, Graph);
 }
 
-void UAvaPlaybackEditorGraphSchema::DroppedAssetsOnNode(const TArray<FAssetData>& Assets, const FVector2D& GraphPosition
+void UAvaPlaybackEditorGraphSchema::DroppedAssetsOnNode(const TArray<FAssetData>& Assets, const FVector2f& GraphPosition
 	, UEdGraphNode* Node) const
 {
 	Super::DroppedAssetsOnNode(Assets, GraphPosition, Node);
diff --git a/Engine/Plugins/Experimental/Avalanche/Source/AvalancheMediaEditor/Private/Playback/Graph/AvaPlaybackEditorGraphSchema.h b/Engine/Plugins/Experimental/Avalanche/Source/AvalancheMediaEditor/Private/Playback/Graph/AvaPlaybackEditorGraphSchema.h
--- a/Engine/Plugins/Experimental/Avalanche/Source/AvalancheMediaEditor/Private/Playback/Graph/AvaPlaybackEditorGraphSchema.h
+++ b/Engine/Plugins/Experimental/Avalanche/Source/AvalancheMediaEditor/Private/Playback/Graph/AvaPlaybackEditorGraphSchema.h
@@ -56,8 +56,8 @@ public:
 	virtual void BreakNodeLinks(UEdGraphNode& TargetNode) const override;
 	virtual void BreakPinLinks(UEdGraphPin& TargetPin, bool bSendsNodeNotifcation) const override;
 	virtual void GetAssetsGraphHoverMessage(const TArray<FAssetData>& Assets, const UEdGraph* HoverGraph, FString& OutTooltipText, bool& OutOkIcon) const override;
-	virtual void DroppedAssetsOnGraph(const TArray<struct FAssetData>& Assets, const FVector2D& GraphPosition, UEdGraph* Graph) const override;
-	virtual void DroppedAssetsOnNode(const TArray<FAssetData>& Assets, const FVector2D& GraphPosition, UEdGraphNode* Node) const override;
+	virtual void DroppedAssetsOnGraph(const TArray<struct FAssetData>& Assets, const FVector2f& GraphPosition, UEdGraph* Graph) const override;
+	virtual void DroppedAssetsOnNode(const TArray<FAssetData>& Assets, const FVector2f& GraphPosition, UEdGraphNode* Node) const override;
 	virtual int32 GetNodeSelectionCount(const UEdGraph* Graph) const override;
 	virtual TSharedPtr<FEdGraphSchemaAction> GetCreateCommentAction() const override;
 	//~ End UEdGraphSchema
diff --git a/Engine/Plugins/Experimental/Avalanche/Source/AvalancheMediaEditor/Private/Playback/Graph/SchemaActions/AvaPlaybackAction_NewComment.cpp b/Engine/Plugins/Experimental/Avalanche/Source/AvalancheMediaEditor/Private/Playback/Graph/SchemaActions/AvaPlaybackAction_NewComment.cpp
--- a/Engine/Plugins/Experimental/Avalanche/Source/AvalancheMediaEditor/Private/Playback/Graph/SchemaActions/AvaPlaybackAction_NewComment.cpp
+++ b/Engine/Plugins/Experimental/Avalanche/Source/AvalancheMediaEditor/Private/Playback/Graph/SchemaActions/AvaPlaybackAction_NewComment.cpp
@@ -8,12 +8,12 @@
 #include "Playback/IAvaPlaybackGraphEditor.h"
 
 UEdGraphNode* FAvaPlaybackAction_NewComment::PerformAction(UEdGraph* ParentGraph, UEdGraphPin* FromPin
-	, const FVector2D Location, bool bSelectNewNode)
+	, const FVector2f& Location, bool bSelectNewNode)
 {
 	// Add menu item for creating comment boxes
 	UEdGraphNode_Comment* const CommentTemplate = NewObject<UEdGraphNode_Comment>();
 
-	FVector2D SpawnLocation = Location;
+	FVector2f SpawnLocation = Location;
 	
 	UAvaPlaybackGraph* const Playback = CastChecked<UAvaPlaybackEditorGraph>(ParentGraph)->GetPlaybackGraph();
 	check(Playback);
diff --git a/Engine/Plugins/Experimental/Avalanche/Source/AvalancheMediaEditor/Private/Playback/Graph/SchemaActions/AvaPlaybackAction_NewComment.h b/Engine/Plugins/Experimental/Avalanche/Source/AvalancheMediaEditor/Private/Playback/Graph/SchemaActions/AvaPlaybackAction_NewComment.h
--- a/Engine/Plugins/Experimental/Avalanche/Source/AvalancheMediaEditor/Private/Playback/Graph/SchemaActions/AvaPlaybackAction_NewComment.h
+++ b/Engine/Plugins/Experimental/Avalanche/Source/AvalancheMediaEditor/Private/Playback/Graph/SchemaActions/AvaPlaybackAction_NewComment.h
@@ -18,9 +18,10 @@ struct FAvaPlaybackAction_NewComment : public FEdGraphSchemaAction
 	}
 
 	//~ Begin FEdGraphSchemaAction Interface
+	using FEdGraphSchemaAction::PerformAction; // Prevent hiding of deprecated base class function with FVector2D
 	virtual UEdGraphNode* PerformAction(class UEdGraph* ParentGraph
 		, UEdGraphPin* FromPin
-		, const FVector2D Location
+		, const FVector2f& Location
 		, bool bSelectNewNode = true) override;
 	//~ End FEdGraphSchemaAction Interface
 };
diff --git a/Engine/Plugins/Experimental/Avalanche/Source/AvalancheMediaEditor/Private/Playback/Graph/SchemaActions/AvaPlaybackAction_NewNode.cpp b/Engine/Plugins/Experimental/Avalanche/Source/AvalancheMediaEditor/Private/Playback/Graph/SchemaActions/AvaPlaybackAction_NewNode.cpp
--- a/Engine/Plugins/Experimental/Avalanche/Source/AvalancheMediaEditor/Private/Playback/Graph/SchemaActions/AvaPlaybackAction_NewNode.cpp
+++ b/Engine/Plugins/Experimental/Avalanche/Source/AvalancheMediaEditor/Private/Playback/Graph/SchemaActions/AvaPlaybackAction_NewNode.cpp
@@ -17,7 +17,7 @@ FAvaPlaybackAction_NewNode::FAvaPlaybackAction_NewNode(FText InNodeCategory, FTe
 
 UEdGraphNode* FAvaPlaybackAction_NewNode::PerformAction(UEdGraph* ParentGraph
 	, UEdGraphPin* FromPin
-	, const FVector2D Location
+	, const FVector2f& Location
 	, bool bSelectNewNode)
 {
 	check(IsValid(PlaybackNodeClass));
diff --git a/Engine/Plugins/Experimental/Avalanche/Source/AvalancheMediaEditor/Private/Playback/Graph/SchemaActions/AvaPlaybackAction_NewNode.h b/Engine/Plugins/Experimental/Avalanche/Source/AvalancheMediaEditor/Private/Playback/Graph/SchemaActions/AvaPlaybackAction_NewNode.h
--- a/Engine/Plugins/Experimental/Avalanche/Source/AvalancheMediaEditor/Private/Playback/Graph/SchemaActions/AvaPlaybackAction_NewNode.h
+++ b/Engine/Plugins/Experimental/Avalanche/Source/AvalancheMediaEditor/Private/Playback/Graph/SchemaActions/AvaPlaybackAction_NewNode.h
@@ -26,7 +26,7 @@ public:
 	//~ Begin FEdGraphSchemaAction
 	virtual UEdGraphNode* PerformAction(class UEdGraph* ParentGraph
 		, UEdGraphPin* FromPin
-		, const FVector2D Location
+		, const FVector2f& Location
 		, bool bSelectNewNode = true) override;
 	//~ End FEdGraphSchemaAction
 
diff --git a/Engine/Plugins/Experimental/Avalanche/Source/AvalancheMediaEditor/Private/Playback/Graph/SchemaActions/AvaPlaybackAction_PasteNode.cpp b/Engine/Plugins/Experimental/Avalanche/Source/AvalancheMediaEditor/Private/Playback/Graph/SchemaActions/AvaPlaybackAction_PasteNode.cpp
--- a/Engine/Plugins/Experimental/Avalanche/Source/AvalancheMediaEditor/Private/Playback/Graph/SchemaActions/AvaPlaybackAction_PasteNode.cpp
+++ b/Engine/Plugins/Experimental/Avalanche/Source/AvalancheMediaEditor/Private/Playback/Graph/SchemaActions/AvaPlaybackAction_PasteNode.cpp
@@ -7,14 +7,14 @@
 #include "Playback/IAvaPlaybackGraphEditor.h"
 
 UEdGraphNode* FAvaPlaybackAction_PasteNode::PerformAction(UEdGraph* ParentGraph, UEdGraphPin* FromPin
-	, const FVector2D Location, bool bSelectNewNode)
+	, const FVector2f& Location, bool bSelectNewNode)
 {
 	UAvaPlaybackGraph* const Playback = CastChecked<UAvaPlaybackEditorGraph>(ParentGraph)->GetPlaybackGraph();
 	check(Playback);
 	
 	if (TSharedPtr<IAvaPlaybackGraphEditor> GraphEditor = Playback->GetGraphEditor())
 	{
-		GraphEditor->PasteNodesHere(Location);
+		GraphEditor->PasteNodesHere(FDeprecateSlateVector2D(Location));
 	}
 	
 	return nullptr;
diff --git a/Engine/Plugins/Experimental/CommonConversation/Source/CommonConversationEditor/Private/ConversationEditorTabFactories.cpp b/Engine/Plugins/Experimental/CommonConversation/Source/CommonConversationEditor/Private/ConversationEditorTabFactories.cpp
--- a/Engine/Plugins/Experimental/CommonConversation/Source/CommonConversationEditor/Private/ConversationEditorTabFactories.cpp
+++ b/Engine/Plugins/Experimental/CommonConversation/Source/CommonConversationEditor/Private/ConversationEditorTabFactories.cpp
@@ -133,7 +133,7 @@ void FConversationGraphEditorSummoner::SaveState(TSharedPtr<SDockTab> Tab, TShar
 
 	TSharedRef<SGraphEditor> GraphEditor = StaticCastSharedRef<SGraphEditor>(Tab->GetContent());
 
-	FVector2D ViewLocation;
+	FVector2f ViewLocation;
 	float ZoomAmount;
 	GraphEditor->GetViewLocation(ViewLocation, ZoomAmount);
 
diff --git a/Engine/Plugins/Experimental/CommonConversation/Source/CommonConversationGraph/Private/ConversationGraphSchema.cpp b/Engine/Plugins/Experimental/CommonConversation/Source/CommonConversationGraph/Private/ConversationGraphSchema.cpp
--- a/Engine/Plugins/Experimental/CommonConversation/Source/CommonConversationGraph/Private/ConversationGraphSchema.cpp
+++ b/Engine/Plugins/Experimental/CommonConversation/Source/CommonConversationGraph/Private/ConversationGraphSchema.cpp
@@ -130,7 +130,7 @@ bool IsConnectionAllowed(const UEdGraphPin* PinA, const UEdGraphPin* PinB, FText
 //////////////////////////////////////////////////////////////////////
 //
 
-UEdGraphNode* FConversationGraphSchemaAction_AutoArrange::PerformAction(class UEdGraph* ParentGraph, UEdGraphPin* FromPin, const FVector2D Location, bool bSelectNewNode)
+UEdGraphNode* FConversationGraphSchemaAction_AutoArrange::PerformAction(class UEdGraph* ParentGraph, UEdGraphPin* FromPin, const FVector2f& Location, bool bSelectNewNode)
 {
 // 	if (UBehaviorTreeGraph* Graph = Cast<UBehaviorTreeGraph>(ParentGraph))
 // 	{
@@ -533,13 +533,13 @@ const FPinConnectionResponse UConversationGraphSchema::CanMergeNodes(const UEdGr
 	return FPinConnectionResponse(CONNECT_RESPONSE_DISALLOW, TEXT(""));
 }
 
-void UConversationGraphSchema::OnPinConnectionDoubleCicked(UEdGraphPin* PinA, UEdGraphPin* PinB, const FVector2D& GraphPosition) const
+void UConversationGraphSchema::OnPinConnectionDoubleCicked(UEdGraphPin* PinA, UEdGraphPin* PinB, const FVector2f& GraphPosition) const
 {
 	const FScopedTransaction Transaction(LOCTEXT("CreateRerouteNodeOnWire", "Create Reroute Node"));
 
 	//@TODO: This constant is duplicated from inside of SGraphNodeKnot
-	const FVector2D NodeSpacerSize(42.0f, 24.0f);
-	const FVector2D KnotTopLeft = GraphPosition - (NodeSpacerSize * 0.5f);
+	const FVector2f NodeSpacerSize(42.0f, 24.0f);
+	const FVector2f KnotTopLeft = GraphPosition - (NodeSpacerSize * 0.5f);
 
 	// Create a new knot
 	UEdGraph* OwningGraph = PinA->GetOwningNode()->GetGraph();
diff --git a/Engine/Plugins/Experimental/CommonConversation/Source/CommonConversationGraph/Public/ConversationGraphSchema.h b/Engine/Plugins/Experimental/CommonConversation/Source/CommonConversationGraph/Public/ConversationGraphSchema.h
--- a/Engine/Plugins/Experimental/CommonConversation/Source/CommonConversationGraph/Public/ConversationGraphSchema.h
+++ b/Engine/Plugins/Experimental/CommonConversation/Source/CommonConversationGraph/Public/ConversationGraphSchema.h
@@ -32,7 +32,8 @@ struct FConversationGraphSchemaAction_AutoArrange : public FEdGraphSchemaAction
 	{}
 
 	//~ Begin FEdGraphSchemaAction Interface
-	virtual UEdGraphNode* PerformAction(class UEdGraph* ParentGraph, UEdGraphPin* FromPin, const FVector2D Location, bool bSelectNewNode = true) override;
+	using FEdGraphSchemaAction::PerformAction; // Prevent hiding of deprecated base class function with FVector2D
+	virtual UEdGraphNode* PerformAction(class UEdGraph* ParentGraph, UEdGraphPin* FromPin, const FVector2f& Location, bool bSelectNewNode = true) override;
 	//~ End FEdGraphSchemaAction Interface
 };
 
@@ -51,7 +52,7 @@ public:
 	virtual void GetContextMenuActions(class UToolMenu* Menu, class UGraphNodeContextMenuContext* Context) const override;
 	virtual const FPinConnectionResponse CanCreateConnection(const UEdGraphPin* A, const UEdGraphPin* B) const override;
 	virtual const FPinConnectionResponse CanMergeNodes(const UEdGraphNode* A, const UEdGraphNode* B) const override;
-	virtual void OnPinConnectionDoubleCicked(UEdGraphPin* PinA, UEdGraphPin* PinB, const FVector2D& GraphPosition) const override;
+	virtual void OnPinConnectionDoubleCicked(UEdGraphPin* PinA, UEdGraphPin* PinB, const FVector2f& GraphPosition) const override;
 	virtual FLinearColor GetPinTypeColor(const FEdGraphPinType& PinType) const override;
 	virtual class FConnectionDrawingPolicy* CreateConnectionDrawingPolicy(int32 InBackLayerID, int32 InFrontLayerID, float InZoomFactor, const FSlateRect& InClippingRect, class FSlateWindowElementList& InDrawElements, class UEdGraph* InGraphObj) const override;
 	virtual bool IsCacheVisualizationOutOfDate(int32 InVisualizationCacheID) const override;
diff --git a/Engine/Plugins/Experimental/Dataflow/Source/DataflowEditor/Private/Dataflow/DataflowEditorCommands.cpp b/Engine/Plugins/Experimental/Dataflow/Source/DataflowEditor/Private/Dataflow/DataflowEditorCommands.cpp
--- a/Engine/Plugins/Experimental/Dataflow/Source/DataflowEditor/Private/Dataflow/DataflowEditorCommands.cpp
+++ b/Engine/Plugins/Experimental/Dataflow/Source/DataflowEditor/Private/Dataflow/DataflowEditorCommands.cpp
@@ -538,8 +538,8 @@ void FDataflowEditorCommands::DuplicateNodes(UEdGraph* EdGraph, const TSharedPtr
 			}
 		}
 		TArray<UEdGraphNode*> DuplicatedEdNodes;
-		const FVector2D PasteLocation = DataflowGraphEditor->GetPasteLocation();
-		UE::Dataflow::FEditAssetUtils::DuplicateNodes(EdGraph, SelectedEdNodes, PasteLocation, DuplicatedEdNodes);
+		const FVector2f PasteLocation = DataflowGraphEditor->GetPasteLocation2f();
+		UE::Dataflow::FEditAssetUtils::DuplicateNodes(EdGraph, SelectedEdNodes, FDeprecateSlateVector2D(PasteLocation), DuplicatedEdNodes);
 
 		// Update the selection in the Editor
 		if (DuplicatedEdNodes.Num())
@@ -604,8 +604,8 @@ void FDataflowEditorCommands::CopyNodes(UEdGraph* EdGraph, const TSharedPtr<SDat
 void FDataflowEditorCommands::PasteNodes(UEdGraph* EdGraph, const TSharedPtr<SDataflowGraphEditor>& DataflowGraphEditor)
 {
 	TArray<UEdGraphNode*> PastedEdNodes;
-	const FVector2D PasteLocation = DataflowGraphEditor->GetPasteLocation();
-	UE::Dataflow::FEditAssetUtils::PasteNodesFromClipboard(EdGraph, PasteLocation, PastedEdNodes);
+	const FVector2f PasteLocation = DataflowGraphEditor->GetPasteLocation2f();
+	UE::Dataflow::FEditAssetUtils::PasteNodesFromClipboard(EdGraph, FDeprecateSlateVector2D(PasteLocation), PastedEdNodes);
 
 	// select the pasted nodes
 	if (PastedEdNodes.Num())
diff --git a/Engine/Plugins/Experimental/Dataflow/Source/DataflowEditor/Private/Dataflow/DataflowEditorMode.cpp b/Engine/Plugins/Experimental/Dataflow/Source/DataflowEditor/Private/Dataflow/DataflowEditorMode.cpp
--- a/Engine/Plugins/Experimental/Dataflow/Source/DataflowEditor/Private/Dataflow/DataflowEditorMode.cpp
+++ b/Engine/Plugins/Experimental/Dataflow/Source/DataflowEditor/Private/Dataflow/DataflowEditorMode.cpp
@@ -1066,7 +1066,7 @@ UEdGraphNode* UDataflowEditorMode::CreateNewNode(const FName& NewNodeTypeName)
 				FAssetSchemaAction_Dataflow_CreateNode_DataflowEdNode::CreateAction(DataflowGraph, NewNodeTypeName);
 			constexpr UEdGraphPin* FromPin = nullptr;
 			constexpr bool bSelectNewNode = true;
-			return NodeAction->PerformAction(DataflowGraph, FromPin, PinnedDataflowGraphEditor->GetPasteLocation(), bSelectNewNode);
+			return NodeAction->PerformAction(DataflowGraph, FromPin, PinnedDataflowGraphEditor->GetPasteLocation2f(), bSelectNewNode);
 		}
 	}
 	return nullptr;
diff --git a/Engine/Plugins/Experimental/Dataflow/Source/DataflowEditor/Private/Dataflow/DataflowGraphEditor.cpp b/Engine/Plugins/Experimental/Dataflow/Source/DataflowEditor/Private/Dataflow/DataflowGraphEditor.cpp
--- a/Engine/Plugins/Experimental/Dataflow/Source/DataflowEditor/Private/Dataflow/DataflowGraphEditor.cpp
+++ b/Engine/Plugins/Experimental/Dataflow/Source/DataflowEditor/Private/Dataflow/DataflowGraphEditor.cpp
@@ -637,7 +637,7 @@ void SDataflowGraphEditor::CreateComment()
 		const TSharedPtr<SGraphEditor>& InGraphEditor = SharedThis(GetGraphEditor());
 
 		TSharedPtr<FAssetSchemaAction_Dataflow_CreateCommentNode_DataflowEdNode> CommentAction = FAssetSchemaAction_Dataflow_CreateCommentNode_DataflowEdNode::CreateAction(EdGraph.Get(), InGraphEditor);
-		CommentAction->PerformAction(EdGraph.Get(), nullptr, GetGraphEditor()->GetPasteLocation(), false);
+		CommentAction->PerformAction(EdGraph.Get(), nullptr, GetGraphEditor()->GetPasteLocation2f(), false);
 	}
 }
 
diff --git a/Engine/Plugins/Experimental/Dataflow/Source/DataflowEditor/Private/Dataflow/DataflowGraphSchemaAction.cpp b/Engine/Plugins/Experimental/Dataflow/Source/DataflowEditor/Private/Dataflow/DataflowGraphSchemaAction.cpp
--- a/Engine/Plugins/Experimental/Dataflow/Source/DataflowEditor/Private/Dataflow/DataflowGraphSchemaAction.cpp
+++ b/Engine/Plugins/Experimental/Dataflow/Source/DataflowEditor/Private/Dataflow/DataflowGraphSchemaAction.cpp
@@ -51,10 +51,10 @@ TSharedPtr<FAssetSchemaAction_Dataflow_CreateNode_DataflowEdNode> FAssetSchemaAc
 	return TSharedPtr<FAssetSchemaAction_Dataflow_CreateNode_DataflowEdNode>(nullptr);
 }
 
-UEdGraphNode* FAssetSchemaAction_Dataflow_CreateNode_DataflowEdNode::PerformAction(class UEdGraph* ParentGraph, UEdGraphPin* FromPin, const FVector2D Location, bool bSelectNewNode)
+UEdGraphNode* FAssetSchemaAction_Dataflow_CreateNode_DataflowEdNode::PerformAction(class UEdGraph* ParentGraph, UEdGraphPin* FromPin, const FVector2f& Location, bool bSelectNewNode)
 {
 	// by default use the type name 
-	UEdGraphNode* NewNode = UE::Dataflow::FEditAssetUtils::AddNewNode(ParentGraph, Location, NodeName, NodeTypeName, FromPin);
+	UEdGraphNode* NewNode = UE::Dataflow::FEditAssetUtils::AddNewNode(ParentGraph, FDeprecateSlateVector2D(Location), NodeName, NodeTypeName, FromPin);
 
 	if (NewNode && bSelectNewNode)
 	{
@@ -180,7 +180,7 @@ void FEdGraphSchemaAction_DataflowVariable::DuplicateItem()
 }
 
 /** Execute this action, given the graph and schema, and possibly a pin that we were dragged from. Returns a node that was created by this action (if any). */
-UEdGraphNode* FEdGraphSchemaAction_DataflowVariable::PerformAction(class UEdGraph* ParentGraph, UEdGraphPin* FromPin, const FVector2D Location, bool bSelectNewNode)
+UEdGraphNode* FEdGraphSchemaAction_DataflowVariable::PerformAction(class UEdGraph* ParentGraph, UEdGraphPin* FromPin, const FVector2f& Location, bool bSelectNewNode)
 {
 	// TODO : move this code to the FEditableDataflow eventually 
 	UEdGraphNode* NewEdNode = nullptr;
@@ -282,7 +282,7 @@ void FEdGraphSchemaAction_DataflowSubGraph::SetForEachSubGraph(bool bValue)
 	}
 }
 
-UEdGraphNode* FEdGraphSchemaAction_DataflowSubGraph::PerformAction(class UEdGraph* ParentGraph, UEdGraphPin* FromPin, const FVector2D Location, bool bSelectNewNode)
+UEdGraphNode* FEdGraphSchemaAction_DataflowSubGraph::PerformAction(class UEdGraph* ParentGraph, UEdGraphPin* FromPin, const FVector2f& Location, bool bSelectNewNode)
 {
 	UEdGraphNode* NewEdNode = nullptr;
 	if (TSharedPtr<FAssetSchemaAction_Dataflow_CreateNode_DataflowEdNode>  SubGraphNodeAction = FAssetSchemaAction_Dataflow_CreateNode_DataflowEdNode::CreateAction(ParentGraph, FDataflowCallSubGraphNode::StaticType(), GetSubGraphName()))
diff --git a/Engine/Plugins/Experimental/Dataflow/Source/DataflowEditor/Private/Dataflow/DataflowSCommentNode.cpp b/Engine/Plugins/Experimental/Dataflow/Source/DataflowEditor/Private/Dataflow/DataflowSCommentNode.cpp
--- a/Engine/Plugins/Experimental/Dataflow/Source/DataflowEditor/Private/Dataflow/DataflowSCommentNode.cpp
+++ b/Engine/Plugins/Experimental/Dataflow/Source/DataflowEditor/Private/Dataflow/DataflowSCommentNode.cpp
@@ -26,7 +26,7 @@ TSharedPtr<FAssetSchemaAction_Dataflow_CreateCommentNode_DataflowEdNode> FAssetS
 //
 //  Created comment node
 //
-UEdGraphNode* FAssetSchemaAction_Dataflow_CreateCommentNode_DataflowEdNode::PerformAction(class UEdGraph* ParentGraph, UEdGraphPin* FromPin, const FVector2D Location, bool bSelectNewNode)
+UEdGraphNode* FAssetSchemaAction_Dataflow_CreateCommentNode_DataflowEdNode::PerformAction(class UEdGraph* ParentGraph, UEdGraphPin* FromPin, const FVector2f& Location, bool bSelectNewNode)
 {
 	constexpr float MinSizeX = 500.f;
 	constexpr float MinSizeY = 250.f;
@@ -34,10 +34,10 @@ UEdGraphNode* FAssetSchemaAction_Dataflow_CreateCommentNode_DataflowEdNode::Perf
 	// set default bounds 
 	FSlateRect Bounds
 	{
-		(float)Location.X,
-		(float)Location.Y,
-		(float)Location.X + MinSizeX,
-		(float)Location.Y + MinSizeY
+		Location.X,
+		Location.Y,
+		Location.X + MinSizeX,
+		Location.Y + MinSizeY
 	};
 	GraphEditor->GetBoundsForSelectedNodes(Bounds, 50.0f);
 
diff --git a/Engine/Plugins/Experimental/Dataflow/Source/DataflowEditor/Private/Dataflow/DataflowSchema.cpp b/Engine/Plugins/Experimental/Dataflow/Source/DataflowEditor/Private/Dataflow/DataflowSchema.cpp
--- a/Engine/Plugins/Experimental/Dataflow/Source/DataflowEditor/Private/Dataflow/DataflowSchema.cpp
+++ b/Engine/Plugins/Experimental/Dataflow/Source/DataflowEditor/Private/Dataflow/DataflowSchema.cpp
@@ -391,12 +391,12 @@ void UDataflowSchema::GetAssetsGraphHoverMessage(const TArray<FAssetData>&Assets
 		}
 	}
 }
-void UDataflowSchema::DroppedAssetsOnGraph(const TArray<FAssetData>&Assets, const FVector2D & GraphPosition, UEdGraph * Graph) const
+void UDataflowSchema::DroppedAssetsOnGraph(const TArray<FAssetData>& Assets, const FVector2f& GraphPosition, UEdGraph* Graph) const
 {
 	using FDataflowNodeFactory = UE::Dataflow::FNodeFactory;
 
-	const FVector2D NodeOffset{ 0,100 };
-	FVector2D NodePosition = GraphPosition;
+	const FVector2f NodeOffset{ 0.0f,100.0f };
+	FVector2f NodePosition = GraphPosition;
 
 	for (const FAssetData& AssetData: Assets)
 	{
@@ -405,7 +405,7 @@ void UDataflowSchema::DroppedAssetsOnGraph(const TArray<FAssetData>&Assets, cons
 			const FName GetterNodeType = FDataflowNodeFactory::GetInstance()->GetGetterNodeFromAssetClass(*AssetObject->GetClass());
 			if (!GetterNodeType.IsNone())
 			{
-				if (UDataflowEdNode* DataflowEdNode = UE::Dataflow::FEditAssetUtils::AddNewNode(Graph, NodePosition, AssetObject->GetFName(), GetterNodeType, nullptr))
+				if (UDataflowEdNode* DataflowEdNode = UE::Dataflow::FEditAssetUtils::AddNewNode(Graph, FDeprecateSlateVector2D(NodePosition), AssetObject->GetFName(), GetterNodeType, nullptr))
 				{
 					NodePosition += NodeOffset;
 					if (TSharedPtr<FDataflowNode> DataflowNode = DataflowEdNode->GetDataflowNode())
@@ -440,7 +440,7 @@ void UDataflowSchema::GetAssetsNodeHoverMessage(const TArray<struct FAssetData>&
 	}
 }
 
-void UDataflowSchema::DroppedAssetsOnNode(const TArray<struct FAssetData>& Assets, const FVector2D& GraphPosition, UEdGraphNode* Node) const
+void UDataflowSchema::DroppedAssetsOnNode(const TArray<struct FAssetData>& Assets, const FVector2f& GraphPosition, UEdGraphNode* Node) const
 {
 	if (TSharedPtr<FDataflowNode> DataflowNode = UDataflowEdNode::GetDataflowNodeFromEdNode(Node))
 	{
@@ -530,7 +530,7 @@ float UDataflowSchema::GetPinTypeWireThickness(const FName& Type) const
 	return CDefaultWireThickness;
 }
 
-static void CreateAndConnectNewReRouteNode(UEdGraphPin* FromPin, UEdGraphPin* ToPin, const FVector2D& GraphPosition)
+static void CreateAndConnectNewReRouteNode(UEdGraphPin* FromPin, UEdGraphPin* ToPin, const UE::Slate::FDeprecateVector2DParameter& GraphPosition)
 {
 	const UEdGraphNode* FromNode = FromPin->GetOwningNode();
 	UEdGraph* EdGraph = FromNode->GetGraph();
@@ -556,9 +556,9 @@ static void CreateAndConnectNewReRouteNode(UEdGraphPin* FromPin, UEdGraphPin* To
 }
 
 
-void UDataflowSchema::OnPinConnectionDoubleCicked(UEdGraphPin* PinA, UEdGraphPin * PinB, const FVector2D & GraphPosition) const
+void UDataflowSchema::OnPinConnectionDoubleCicked(UEdGraphPin* PinA, UEdGraphPin* PinB, const FVector2f& GraphPosition) const
 {
-	CreateAndConnectNewReRouteNode(PinA, PinB, GraphPosition);
+	CreateAndConnectNewReRouteNode(PinA, PinB, FDeprecateSlateVector2D(GraphPosition));
 }
 
 void UDataflowSchema::BreakPinLinks(UEdGraphPin& TargetPin, bool bSendsNodeNotifcation) const
diff --git a/Engine/Plugins/Experimental/Dataflow/Source/DataflowEditor/Public/Dataflow/DataflowGraphSchemaAction.h b/Engine/Plugins/Experimental/Dataflow/Source/DataflowEditor/Public/Dataflow/DataflowGraphSchemaAction.h
--- a/Engine/Plugins/Experimental/Dataflow/Source/DataflowEditor/Public/Dataflow/DataflowGraphSchemaAction.h
+++ b/Engine/Plugins/Experimental/Dataflow/Source/DataflowEditor/Public/Dataflow/DataflowGraphSchemaAction.h
@@ -43,7 +43,8 @@ public:
 
 	static TSharedPtr<FAssetSchemaAction_Dataflow_CreateNode_DataflowEdNode> CreateAction(const UEdGraph* ParentGraph, const FName& NodeTypeName, const FName& InOverrideNodeName = NAME_None);
 
-	virtual UEdGraphNode* PerformAction(class UEdGraph* ParentGraph, UEdGraphPin* FromPin, const FVector2D Location, bool bSelectNewNode = true) override;
+	using FEdGraphSchemaAction::PerformAction; // Prevent hiding of deprecated base class function with FVector2D
+	virtual UEdGraphNode* PerformAction(class UEdGraph* ParentGraph, UEdGraphPin* FromPin, const FVector2f& Location, bool bSelectNewNode = true) override;
 
 	FName NodeName;
 	FName NodeTypeName;
@@ -84,7 +85,7 @@ public:
 
 	// FEdGraphSchemaAction interface
 	virtual bool IsAVariable() const { return true; }
-	virtual UEdGraphNode* PerformAction(class UEdGraph* ParentGraph, UEdGraphPin* FromPin, const FVector2D Location, bool bSelectNewNode = true) override;
+	virtual UEdGraphNode* PerformAction(class UEdGraph* ParentGraph, UEdGraphPin* FromPin, const FVector2f& Location, bool bSelectNewNode = true) override;
 	// End of FEdGraphSchemaAction interface
 
 	bool CanRenameItem(FText NewName) const;
@@ -173,7 +174,7 @@ public:
 
 	// FEdGraphSchemaAction interface
 	virtual bool IsAVariable() const { return false; }
-	virtual UEdGraphNode* PerformAction(class UEdGraph* ParentGraph, UEdGraphPin* FromPin, const FVector2D Location, bool bSelectNewNode = true) override;
+	virtual UEdGraphNode* PerformAction(class UEdGraph* ParentGraph, UEdGraphPin* FromPin, const FVector2f& Location, bool bSelectNewNode = true) override;
 	// End of FEdGraphSchemaAction interface
 
 	bool CanRenameItem(FText NewName) const;
diff --git a/Engine/Plugins/Experimental/Dataflow/Source/DataflowEditor/Public/Dataflow/DataflowSCommentNode.h b/Engine/Plugins/Experimental/Dataflow/Source/DataflowEditor/Public/Dataflow/DataflowSCommentNode.h
--- a/Engine/Plugins/Experimental/Dataflow/Source/DataflowEditor/Public/Dataflow/DataflowSCommentNode.h
+++ b/Engine/Plugins/Experimental/Dataflow/Source/DataflowEditor/Public/Dataflow/DataflowSCommentNode.h
@@ -40,7 +40,8 @@ public:
 
 	static TSharedPtr<FAssetSchemaAction_Dataflow_CreateCommentNode_DataflowEdNode> CreateAction(UEdGraph* ParentGraph, const TSharedPtr<SGraphEditor>& GraphEditor);
 
-	virtual UEdGraphNode* PerformAction(class UEdGraph* ParentGraph, UEdGraphPin* FromPin, const FVector2D Location, bool bSelectNewNode = true) override;
+	using FEdGraphSchemaAction::PerformAction; // Prevent hiding of deprecated base class function with FVector2D
+	virtual UEdGraphNode* PerformAction(class UEdGraph* ParentGraph, UEdGraphPin* FromPin, const FVector2f& Location, bool bSelectNewNode = true) override;
 	//virtual void AddReferencedObjects(FReferenceCollector& Collector) override;
 
 	TSharedPtr<SGraphEditor> GraphEditor;
diff --git a/Engine/Plugins/Experimental/Dataflow/Source/DataflowEditor/Public/Dataflow/DataflowSchema.h b/Engine/Plugins/Experimental/Dataflow/Source/DataflowEditor/Public/Dataflow/DataflowSchema.h
--- a/Engine/Plugins/Experimental/Dataflow/Source/DataflowEditor/Public/Dataflow/DataflowSchema.h
+++ b/Engine/Plugins/Experimental/Dataflow/Source/DataflowEditor/Public/Dataflow/DataflowSchema.h
@@ -29,15 +29,15 @@ public:
 
 	virtual FLinearColor GetPinTypeColor(const FEdGraphPinType& PinType) const override;
 	virtual FConnectionDrawingPolicy* CreateConnectionDrawingPolicy(int32 InBackLayerID, int32 InFrontLayerID, float InZoomFactor, const FSlateRect& InClippingRect, class FSlateWindowElementList& InDrawElements, class UEdGraph* InGraphObj) const override;
-	virtual void OnPinConnectionDoubleCicked(UEdGraphPin* PinA, UEdGraphPin* PinB, const FVector2D& GraphPosition) const override;
+	virtual void OnPinConnectionDoubleCicked(UEdGraphPin* PinA, UEdGraphPin* PinB, const FVector2f& GraphPosition) const override;
 	virtual void BreakPinLinks(UEdGraphPin& TargetPin, bool bSendsNodeNotifcation) const override;
 	virtual bool TryCreateConnection(UEdGraphPin* PinA, UEdGraphPin* PinB) const override;
 
 	virtual void GetAssetsGraphHoverMessage(const TArray<FAssetData>& Assets, const UEdGraph* HoverGraph, FString& OutTooltipText, bool& OutOkIcon) const override;
-	virtual void DroppedAssetsOnGraph(const TArray<FAssetData>& Assets, const FVector2D& GraphPosition, UEdGraph* Graph) const override;
+	virtual void DroppedAssetsOnGraph(const TArray<FAssetData>& Assets, const FVector2f& GraphPosition, UEdGraph* Graph) const override;
 
 	virtual void GetAssetsNodeHoverMessage(const TArray<struct FAssetData>& Assets, const UEdGraphNode* HoverNode, FString& OutTooltipText, bool& OutOkIcon) const override;
-	virtual void DroppedAssetsOnNode(const TArray<struct FAssetData>& Assets, const FVector2D& GraphPosition, UEdGraphNode* Node) const override;
+	virtual void DroppedAssetsOnNode(const TArray<struct FAssetData>& Assets, const FVector2f& GraphPosition, UEdGraphNode* Node) const override;
 
 	virtual void SetPinBeingDroppedOnNode(UEdGraphPin* InSourcePin) const override { PinBeingDropped = InSourcePin; }
 	virtual bool SupportsDropPinOnNode(UEdGraphNode* InTargetNode, const FEdGraphPinType& InSourcePinType, EEdGraphPinDirection InSourcePinDirection, FText& OutMessage) const override;
diff --git a/Engine/Plugins/Experimental/PluginReferenceViewer/Source/PluginReferenceViewer/Private/PluginReferenceViewerSchema.cpp b/Engine/Plugins/Experimental/PluginReferenceViewer/Source/PluginReferenceViewer/Private/PluginReferenceViewerSchema.cpp
--- a/Engine/Plugins/Experimental/PluginReferenceViewer/Source/PluginReferenceViewer/Private/PluginReferenceViewerSchema.cpp
+++ b/Engine/Plugins/Experimental/PluginReferenceViewer/Source/PluginReferenceViewer/Private/PluginReferenceViewerSchema.cpp
@@ -175,7 +175,7 @@ FConnectionDrawingPolicy* UPluginReferenceViewerSchema::CreateConnectionDrawingP
 	return new FReferenceViewerConnectionDrawingPolicy(InBackLayerID, InFrontLayerID, InZoomFactor, InClippingRect, InDrawElements);
 }
 
-void UPluginReferenceViewerSchema::DroppedAssetsOnGraph(const TArray<FAssetData>& Assets, const FVector2D& GraphPosition, UEdGraph* Graph) const
+void UPluginReferenceViewerSchema::DroppedAssetsOnGraph(const TArray<FAssetData>& Assets, const FVector2f& GraphPosition, UEdGraph* Graph) const
 {
 }
 
diff --git a/Engine/Plugins/Experimental/PluginReferenceViewer/Source/PluginReferenceViewer/Private/PluginReferenceViewerSchema.h b/Engine/Plugins/Experimental/PluginReferenceViewer/Source/PluginReferenceViewer/Private/PluginReferenceViewerSchema.h
--- a/Engine/Plugins/Experimental/PluginReferenceViewer/Source/PluginReferenceViewer/Private/PluginReferenceViewerSchema.h
+++ b/Engine/Plugins/Experimental/PluginReferenceViewer/Source/PluginReferenceViewer/Private/PluginReferenceViewerSchema.h
@@ -26,7 +26,7 @@ public:
 	virtual FPinConnectionResponse MovePinLinks(UEdGraphPin& MoveFromPin, UEdGraphPin& MoveToPin, bool bIsIntermediateMove = false, bool bNotifyLinkedNodes = false) const override;
 	virtual FPinConnectionResponse CopyPinLinks(UEdGraphPin& CopyFromPin, UEdGraphPin& CopyToPin, bool bIsIntermediateCopy = false) const override;
 	virtual class FConnectionDrawingPolicy* CreateConnectionDrawingPolicy(int32 InBackLayerID, int32 InFrontLayerID, float InZoomFactor, const FSlateRect& InClippingRect, class FSlateWindowElementList& InDrawElements, class UEdGraph* InGraphObj) const override;
-	virtual void DroppedAssetsOnGraph(const TArray<FAssetData>& Assets, const FVector2D& GraphPosition, UEdGraph* Graph) const override;
+	virtual void DroppedAssetsOnGraph(const TArray<FAssetData>& Assets, const FVector2f& GraphPosition, UEdGraph* Graph) const override;
 	virtual void GetAssetsGraphHoverMessage(const TArray<FAssetData>& Assets, const UEdGraph* HoverGraph, FString& OutTooltipText, bool& OutOkIcon) const override;
 	// End of UEdGraphSchema interface
 };
\ No newline at end of file
diff --git a/Engine/Plugins/Experimental/Workspace/Source/WorkspaceEditor/Private/GraphDocumentState.h b/Engine/Plugins/Experimental/Workspace/Source/WorkspaceEditor/Private/GraphDocumentState.h
--- a/Engine/Plugins/Experimental/Workspace/Source/WorkspaceEditor/Private/GraphDocumentState.h
+++ b/Engine/Plugins/Experimental/Workspace/Source/WorkspaceEditor/Private/GraphDocumentState.h
@@ -1,8 +1,9 @@
-﻿// Copyright Epic Games, Inc. All Rights Reserved.
+// Copyright Epic Games, Inc. All Rights Reserved.
 
 #pragma once
 
 #include "WorkspaceDocumentState.h"
+#include "Types/SlateVector2.h"
 #include "GraphDocumentState.generated.h"
 
 // Base struct used to persist workspace document state
@@ -13,14 +14,14 @@ struct FGraphDocumentState : public FWorkspaceDocumentState
 
 	FGraphDocumentState() = default;
 
-	FGraphDocumentState(const UObject* InObject, const FWorkspaceOutlinerItemExport& InExport, const FVector2D& InViewLocation, float InZoomAmount)
+	FGraphDocumentState(const UObject* InObject, const FWorkspaceOutlinerItemExport& InExport, const UE::Slate::FDeprecateVector2DParameter& InViewLocation, float InZoomAmount)
 		: FWorkspaceDocumentState(InObject, InExport)
 		, ViewLocation(InViewLocation)
 		, ZoomAmount(InZoomAmount)
 	{}
 
 	UPROPERTY()
-	FVector2D ViewLocation = FVector2D::Zero();
+	FDeprecateSlateVector2D ViewLocation = FVector2f::Zero();
 
 	UPROPERTY()
 	float ZoomAmount = 0.f;
diff --git a/Engine/Plugins/Experimental/Workspace/Source/WorkspaceEditor/Private/SGraphDocument.cpp b/Engine/Plugins/Experimental/Workspace/Source/WorkspaceEditor/Private/SGraphDocument.cpp
--- a/Engine/Plugins/Experimental/Workspace/Source/WorkspaceEditor/Private/SGraphDocument.cpp
+++ b/Engine/Plugins/Experimental/Workspace/Source/WorkspaceEditor/Private/SGraphDocument.cpp
@@ -156,7 +156,7 @@ void SGraphDocument::PasteNodes()
 	FString TextToImport;
 	FPlatformApplicationMisc::ClipboardPaste(TextToImport);
 
-	OnPasteNodes.ExecuteIfBound(FWorkspaceEditorContext(HostingAppPtr.Pin().ToSharedRef(), Document), GraphEditor->GetPasteLocation(), TextToImport);
+	OnPasteNodes.ExecuteIfBound(FWorkspaceEditorContext(HostingAppPtr.Pin().ToSharedRef(), Document), GraphEditor->GetPasteLocation2f(), TextToImport);
 }
 
 bool SGraphDocument::CanDuplicateSelectedNodes() const
@@ -170,7 +170,7 @@ bool SGraphDocument::CanDuplicateSelectedNodes() const
 void SGraphDocument::DuplicateSelectedNodes()
 {
 	const FGraphPanelSelectionSet SelectedNodes = GraphEditor->GetSelectedNodes();
-	OnDuplicateSelectedNodes.ExecuteIfBound(FWorkspaceEditorContext(HostingAppPtr.Pin().ToSharedRef(), Document), GraphEditor->GetPasteLocation(), SelectedNodes);
+	OnDuplicateSelectedNodes.ExecuteIfBound(FWorkspaceEditorContext(HostingAppPtr.Pin().ToSharedRef(), Document), GraphEditor->GetPasteLocation2f(), SelectedNodes);
 }
 
 bool SGraphDocument::CanSelectAllNodes() const
diff --git a/Engine/Plugins/Experimental/Workspace/Source/WorkspaceEditor/Private/WorkspaceEditorModule.cpp b/Engine/Plugins/Experimental/Workspace/Source/WorkspaceEditor/Private/WorkspaceEditorModule.cpp
--- a/Engine/Plugins/Experimental/Workspace/Source/WorkspaceEditor/Private/WorkspaceEditorModule.cpp
+++ b/Engine/Plugins/Experimental/Workspace/Source/WorkspaceEditor/Private/WorkspaceEditorModule.cpp
@@ -141,7 +141,7 @@ FObjectDocumentArgs FWorkspaceEditorModule::CreateGraphDocumentArgs(const FGraph
 	});
 	Args.OnGetDocumentState = FOnGetDocumentState::CreateLambda([](const FWorkspaceEditorContext& InContext, TSharedRef<SWidget> InWidget)
 	{
-		FVector2D ViewLocation = FVector2d::ZeroVector;
+		FVector2f ViewLocation = FVector2f::ZeroVector;
 		float ZoomAmount = 0.f;
 		
 		if(const TSharedPtr<SGraphDocument> GraphDocument = StaticCastSharedRef<SGraphDocument>(InWidget))
@@ -149,7 +149,7 @@ FObjectDocumentArgs FWorkspaceEditorModule::CreateGraphDocumentArgs(const FGraph
 			GraphDocument->GraphEditor->GetViewLocation(ViewLocation, ZoomAmount);
 		}
 	
-		return TInstancedStruct<FGraphDocumentState>::Make(InContext.Document.Object, InContext.Document.Export, ViewLocation, ZoomAmount);
+		return TInstancedStruct<FGraphDocumentState>::Make(InContext.Document.Object, InContext.Document.Export, FDeprecateSlateVector2D(ViewLocation), ZoomAmount);
 	});
 	Args.OnSetDocumentState = FOnSetDocumentState::CreateLambda([](const FWorkspaceEditorContext& InContext, TSharedRef<SWidget> InWidget, const TInstancedStruct<FWorkspaceDocumentState>& InDocumentState)
 	{
diff --git a/Engine/Plugins/FX/Niagara/Source/Niagara/Public/NiagaraEditorDataBase.h b/Engine/Plugins/FX/Niagara/Source/Niagara/Public/NiagaraEditorDataBase.h
--- a/Engine/Plugins/FX/Niagara/Source/Niagara/Public/NiagaraEditorDataBase.h
+++ b/Engine/Plugins/FX/Niagara/Source/Niagara/Public/NiagaraEditorDataBase.h
@@ -4,6 +4,7 @@
 
 #include "CoreTypes.h"
 #include "NiagaraCommon.h"
+#include "Types/SlateVector2.h"
 
 #include "NiagaraEditorDataBase.generated.h"
 
@@ -16,26 +17,26 @@ struct FNiagaraGraphViewSettings
 	GENERATED_BODY()
 public:
 	FNiagaraGraphViewSettings()
-		: Location(FVector2D::ZeroVector)
+		: Location(FVector2f::ZeroVector)
 		, Zoom(0.0f)
 		, bIsValid(false)
 	{
 	}
 
-	FNiagaraGraphViewSettings(const FVector2D& InLocation, float InZoom)
+	FNiagaraGraphViewSettings(const UE::Slate::FDeprecateVector2DParameter& InLocation, float InZoom)
 		: Location(InLocation)
 		, Zoom(InZoom)
 		, bIsValid(true)
 	{
 	}
 
-	const FVector2D& GetLocation() const { return Location; }
+	const FDeprecateSlateVector2D& GetLocation() const { return Location; }
 	float GetZoom() const { return Zoom; }
 	bool IsValid() const { return bIsValid; }
 
 private:
 	UPROPERTY()
-	FVector2D Location;
+	FDeprecateSlateVector2D Location;
 
 	UPROPERTY()
 	float Zoom;
diff --git a/Engine/Plugins/FX/Niagara/Source/NiagaraEditor/Private/Customizations/NiagaraBakerSettingsDetails.h b/Engine/Plugins/FX/Niagara/Source/NiagaraEditor/Private/Customizations/NiagaraBakerSettingsDetails.h
--- a/Engine/Plugins/FX/Niagara/Source/NiagaraEditor/Private/Customizations/NiagaraBakerSettingsDetails.h
+++ b/Engine/Plugins/FX/Niagara/Source/NiagaraEditor/Private/Customizations/NiagaraBakerSettingsDetails.h
@@ -31,7 +31,7 @@ struct FNiagaraBakerTextureSourceAction : public FEdGraphSchemaAction
 	{}
 
 	//~ Begin FEdGraphSchemaAction Interface
-	virtual UEdGraphNode* PerformAction(class UEdGraph* ParentGraph, UEdGraphPin* FromPin, const FVector2D Location, bool bSelectNewNode = true) override
+	virtual UEdGraphNode* PerformAction(class UEdGraph* ParentGraph, UEdGraphPin* FromPin, const FVector2f& Location, bool bSelectNewNode = true) override
 	{
 		return nullptr;
 	}
diff --git a/Engine/Plugins/FX/Niagara/Source/NiagaraEditor/Private/Customizations/NiagaraTypeCustomizations.h b/Engine/Plugins/FX/Niagara/Source/NiagaraEditor/Private/Customizations/NiagaraTypeCustomizations.h
--- a/Engine/Plugins/FX/Niagara/Source/NiagaraEditor/Private/Customizations/NiagaraTypeCustomizations.h
+++ b/Engine/Plugins/FX/Niagara/Source/NiagaraEditor/Private/Customizations/NiagaraTypeCustomizations.h
@@ -147,7 +147,7 @@ struct FNiagaraStackAssetAction_VarBind : public FEdGraphSchemaAction
 	{}
 
 	//~ Begin FEdGraphSchemaAction Interface
-	virtual UEdGraphNode* PerformAction(class UEdGraph* ParentGraph, UEdGraphPin* FromPin, const FVector2D Location, bool bSelectNewNode = true) override
+	virtual UEdGraphNode* PerformAction(class UEdGraph* ParentGraph, UEdGraphPin* FromPin, const FVector2f& Location, bool bSelectNewNode = true) override
 	{
 		return nullptr;
 	}
diff --git a/Engine/Plugins/FX/Niagara/Source/NiagaraEditor/Private/EdGraphSchema_Niagara.cpp b/Engine/Plugins/FX/Niagara/Source/NiagaraEditor/Private/EdGraphSchema_Niagara.cpp
--- a/Engine/Plugins/FX/Niagara/Source/NiagaraEditor/Private/EdGraphSchema_Niagara.cpp
+++ b/Engine/Plugins/FX/Niagara/Source/NiagaraEditor/Private/EdGraphSchema_Niagara.cpp
@@ -88,7 +88,7 @@ namespace NiagaraNodeNumbers
 	constexpr int32 NiagaraMinNodeDistance = 60;
 }
 
-UEdGraphNode* FNiagaraSchemaAction_NewNode::PerformAction(class UEdGraph* ParentGraph, UEdGraphPin* FromPin, const FVector2D Location, bool bSelectNewNode/* = true*/)
+UEdGraphNode* FNiagaraSchemaAction_NewNode::PerformAction(class UEdGraph* ParentGraph, UEdGraphPin* FromPin, const FVector2f& Location, bool bSelectNewNode/* = true*/)
 {
 	UEdGraphNode* ResultNode = NULL;
 
@@ -148,7 +148,7 @@ UEdGraphNode* FNiagaraSchemaAction_NewNode::PerformAction(class UEdGraph* Parent
 	return ResultNode;
 }
 
-UEdGraphNode* FNiagaraSchemaAction_NewNode::PerformAction(class UEdGraph* ParentGraph, TArray<UEdGraphPin*>& FromPins, const FVector2D Location, bool bSelectNewNode/* = true*/) 
+UEdGraphNode* FNiagaraSchemaAction_NewNode::PerformAction(class UEdGraph* ParentGraph, TArray<UEdGraphPin*>& FromPins, const FVector2f& Location, bool bSelectNewNode/* = true*/) 
 {
 	UEdGraphNode* ResultNode = NULL;
 
@@ -181,12 +181,12 @@ void FNiagaraSchemaAction_NewNode::AddReferencedObjects( FReferenceCollector& Co
 	Collector.AddReferencedObject( NodeTemplate );
 }
 
-UEdGraphNode* FNiagaraSchemaAction_NewComment::PerformAction(class UEdGraph* ParentGraph, UEdGraphPin* FromPin, const FVector2D Location, bool bSelectNewNode /*= true*/)
+UEdGraphNode* FNiagaraSchemaAction_NewComment::PerformAction(class UEdGraph* ParentGraph, UEdGraphPin* FromPin, const FVector2f& Location, bool bSelectNewNode /*= true*/)
 {
 	// Add menu item for creating comment boxes
 	UEdGraphNode_Comment* CommentTemplate = NewObject<UEdGraphNode_Comment>();
 
-	FVector2D SpawnLocation = Location;
+	FVector2f SpawnLocation = Location;
 	FSlateRect Bounds;
 	
 	if (GraphEditor->GetBoundsForSelectedNodes(Bounds, 50.0f))
@@ -1326,13 +1326,13 @@ void UEdGraphSchema_Niagara::ResetPinToAutogeneratedDefaultValue(UEdGraphPin* Pi
 	}
 }
 
-void UEdGraphSchema_Niagara::OnPinConnectionDoubleCicked(UEdGraphPin* PinA, UEdGraphPin* PinB, const FVector2D& GraphPosition) const
+void UEdGraphSchema_Niagara::OnPinConnectionDoubleCicked(UEdGraphPin* PinA, UEdGraphPin* PinB, const FVector2f& GraphPosition) const
 {
 	const FScopedTransaction Transaction(LOCTEXT("CreateRerouteNodeOnWire", "Create Reroute Node"));
 
 	//@TODO: This constant is duplicated from inside of SGraphNodeKnot
-	const FVector2D NodeSpacerSize(42.0f, 24.0f);
-	const FVector2D KnotTopLeft = GraphPosition - (NodeSpacerSize * 0.5f);
+	const FVector2f NodeSpacerSize(42.0f, 24.0f);
+	const FVector2f KnotTopLeft = GraphPosition - (NodeSpacerSize * 0.5f);
 
 	// Create a new knot
 	UEdGraph* ParentGraph = PinA->GetOwningNode()->GetGraph();
@@ -1346,7 +1346,7 @@ void UEdGraphSchema_Niagara::OnPinConnectionDoubleCicked(UEdGraphPin* PinA, UEdG
 }
 
 
-void UEdGraphSchema_Niagara::DroppedAssetsOnGraph(const TArray<FAssetData>&Assets, const FVector2D & GraphPosition, UEdGraph * Graph) const
+void UEdGraphSchema_Niagara::DroppedAssetsOnGraph(const TArray<FAssetData>& Assets, const FVector2f& GraphPosition, UEdGraph* Graph) const
 {
 	uint32 Offset = 0;
 	TArray<UEnum*> Enums;
@@ -2364,7 +2364,7 @@ void UEdGraphSchema_Niagara::PromoteSinglePinToParameter(UEdGraphPin* SourcePin)
 			UEdGraphNode* PinNode = SourcePin->GetOwningNode();
 
 			const float PinVisualOffsetX = 175.0f;
-			InputAction->PerformAction(Graph, SourcePin, FVector2D(PinNode->NodePosX - PinVisualOffsetX, PinNode->NodePosY));
+			InputAction->PerformAction(Graph, SourcePin, FVector2f(PinNode->NodePosX - PinVisualOffsetX, PinNode->NodePosY));
 		}
 	}
 }
diff --git a/Engine/Plugins/FX/Niagara/Source/NiagaraEditor/Private/EdGraphSchema_NiagaraSystemOverview.cpp b/Engine/Plugins/FX/Niagara/Source/NiagaraEditor/Private/EdGraphSchema_NiagaraSystemOverview.cpp
--- a/Engine/Plugins/FX/Niagara/Source/NiagaraEditor/Private/EdGraphSchema_NiagaraSystemOverview.cpp
+++ b/Engine/Plugins/FX/Niagara/Source/NiagaraEditor/Private/EdGraphSchema_NiagaraSystemOverview.cpp
@@ -39,7 +39,7 @@ void UEdGraphSchema_NiagaraSystemOverview::ForceVisualizationCacheClear() const
 	CurrentCacheRefreshID++;
 }
 
-void UEdGraphSchema_NiagaraSystemOverview::DroppedAssetsOnGraph(const TArray<FAssetData>& Assets, const FVector2D& GraphPosition, UEdGraph* Graph) const
+void UEdGraphSchema_NiagaraSystemOverview::DroppedAssetsOnGraph(const TArray<FAssetData>& Assets, const FVector2f& GraphPosition, UEdGraph* Graph) const
 {
 	TArray<UNiagaraEmitter*> Emitters;
 
@@ -61,7 +61,7 @@ void UEdGraphSchema_NiagaraSystemOverview::DroppedAssetsOnGraph(const TArray<FAs
 		{
 			FScopedTransaction AddEmitterTransaction(LOCTEXT("NiagaraEditorDropEmitter", "Niagara: Drag and Drop Emitter"));
 
-			FVector2D NewNodePosition = GraphPosition;
+			FVector2f NewNodePosition = GraphPosition;
 			for (UNiagaraEmitter* Emitter : Emitters)
 			{
 				FGuid Handle = FNiagaraEditorUtilities::AddEmitterToSystem(*System, *Emitter, Emitter->GetExposedVersion().VersionGuid);
diff --git a/Engine/Plugins/FX/Niagara/Source/NiagaraEditor/Private/Widgets/NiagaraScratchScriptEditor.cpp b/Engine/Plugins/FX/Niagara/Source/NiagaraEditor/Private/Widgets/NiagaraScratchScriptEditor.cpp
--- a/Engine/Plugins/FX/Niagara/Source/NiagaraEditor/Private/Widgets/NiagaraScratchScriptEditor.cpp
+++ b/Engine/Plugins/FX/Niagara/Source/NiagaraEditor/Private/Widgets/NiagaraScratchScriptEditor.cpp
@@ -180,7 +180,9 @@ void FNiagaraGraphTabHistory::SaveHistory()
 		TSharedPtr<SNiagaraScratchPadScriptEditor> Editor = GraphEditor.Pin();
 		if (Editor->GetGraphEditor())
 		{
-			Editor->GetGraphEditor()->GetViewLocation(SavedLocation, SavedZoomAmount);
+			FVector2f ViewLocation;
+			Editor->GetGraphEditor()->GetViewLocation(ViewLocation, SavedZoomAmount);
+			SavedLocation = ViewLocation;
 			Editor->GetGraphEditor()->GetViewBookmark(SavedBookmarkId);
 		}
 	}
diff --git a/Engine/Plugins/FX/Niagara/Source/NiagaraEditor/Private/Widgets/NiagaraScratchScriptEditor.h b/Engine/Plugins/FX/Niagara/Source/NiagaraEditor/Private/Widgets/NiagaraScratchScriptEditor.h
--- a/Engine/Plugins/FX/Niagara/Source/NiagaraEditor/Private/Widgets/NiagaraScratchScriptEditor.h
+++ b/Engine/Plugins/FX/Niagara/Source/NiagaraEditor/Private/Widgets/NiagaraScratchScriptEditor.h
@@ -106,7 +106,7 @@ public:
 	 */
 	FNiagaraGraphTabHistory(TSharedPtr<FDocumentTabFactory> InFactory, TSharedPtr<FTabPayload> InPayload)
 		: FGenericTabHistory(InFactory, InPayload)
-		, SavedLocation(FVector2D::ZeroVector)
+		, SavedLocation(FVector2f::ZeroVector)
 		, SavedZoomAmount(INDEX_NONE)
 	{
 
@@ -122,7 +122,7 @@ private:
 	/** The graph editor represented by this history node. While this node is inactive, the graph editor is invalid */
 	TWeakPtr< class SNiagaraScratchPadScriptEditor > GraphEditor;
 	/** Saved location the graph editor was at when this history node was last visited */
-	FVector2D SavedLocation;
+	FDeprecateSlateVector2D SavedLocation;
 	/** Saved zoom the graph editor was at when this history node was last visited */
 	float SavedZoomAmount;
 	/** Saved bookmark ID the graph editor was at when this history node was last visited */
