Avoid using a temporary array to reduce memory churn

diff --git a/Engine/Plugins/Experimental/AnimNext/Source/AnimNext/Public/TransformArray.h b/Engine/Plugins/Experimental/AnimNext/Source/AnimNext/Public/TransformArray.h
--- a/Engine/Plugins/Experimental/AnimNext/Source/AnimNext/Public/TransformArray.h
+++ b/Engine/Plugins/Experimental/AnimNext/Source/AnimNext/Public/TransformArray.h
@@ -5,6 +5,7 @@
 #include "CoreMinimal.h"
 #include "Animation/AnimTypes.h"
 #include "TransformArrayView.h"
+#include "AnimNextStats.h"
 #include <algorithm>
 
 #define ANIM_ENABLE_POINTER_ITERATION 1
@@ -299,6 +300,8 @@ struct TTransformArraySoA
 
 	TTransformArraySoA(const TTransformArraySoA& Other)
 	{
+		SCOPE_CYCLE_COUNTER(STAT_AnimNext_CopyTransforms_SoA);
+
 		AllocatedMemory = Other.AllocatedMemory;
 
 		UpdateViews(AllocatedMemory.GetData(), Other.Num());
@@ -306,9 +309,12 @@ struct TTransformArraySoA
 
 	TTransformArraySoA& operator= (const TTransformArraySoA& Other)
 	{
-		TTransformArraySoA Tmp(Other);
+		SCOPE_CYCLE_COUNTER(STAT_AnimNext_CopyTransforms_SoA);
+
+		AllocatedMemory = Other.AllocatedMemory;
+
+		UpdateViews(AllocatedMemory.GetData(), Other.Num());
 
-		Swap(*this, Tmp);
 		return *this;
 	}
 
@@ -366,6 +372,8 @@ struct TTransformArraySoA
 
 	inline void CopyTransforms(const TTransformArraySoA& Other, int32 Index, int32 NumTransforms)
 	{
+		SCOPE_CYCLE_COUNTER(STAT_AnimNext_CopyTransforms_SoA);
+
 		check(Index < Num() && Index < Other.Num());
 		check(Index + (NumTransforms - 1) < Num() && Index + (NumTransforms - 1) < Other.Num());
 
@@ -376,6 +384,8 @@ struct TTransformArraySoA
 
 	inline void CopyTransforms(const FTransformArraySoAConstView& Other, int32 Index, int32 NumTransforms)
 	{
+		SCOPE_CYCLE_COUNTER(STAT_AnimNext_CopyTransforms_SoA);
+
 		check(Index < Num() && Index < Other.Num());
 		check(Index + (NumTransforms - 1) < Num() && Index + (NumTransforms - 1) < Other.Num());
 
