[AnimNext] Fix build

Missing include

diff --git a/Engine/Plugins/Experimental/AnimNext/Source/AnimNext/Public/ReferencePose.h b/Engine/Plugins/Experimental/AnimNext/Source/AnimNext/Public/ReferencePose.h
--- a/Engine/Plugins/Experimental/AnimNext/Source/AnimNext/Public/ReferencePose.h
+++ b/Engine/Plugins/Experimental/AnimNext/Source/AnimNext/Public/ReferencePose.h
@@ -11,6 +11,7 @@
 #include "ReferenceSkeleton.h"
 #include "BoneIndices.h"
 #include "TransformArray.h"
+#include "Components/SkeletalMeshComponent.h"
 #include "ReferencePose.generated.h"
 
 class USkeletalMesh;
