Add RigUnit to query the recently rendered state of a skinned mesh component

diff --git a/Engine/Plugins/Experimental/AnimNext/Source/AnimNext/Private/Graph/RigUnit_IsRecentlyRendered.cpp b/Engine/Plugins/Experimental/AnimNext/Source/AnimNext/Private/Graph/RigUnit_IsRecentlyRendered.cpp
new file mode 100644
--- /dev/null
+++ b/Engine/Plugins/Experimental/AnimNext/Source/AnimNext/Private/Graph/RigUnit_IsRecentlyRendered.cpp
@@ -0,0 +1,15 @@
+// Copyright Epic Games, Inc. All Rights Reserved.
+
+#include "RigUnit_IsRecentlyRendered.h"
+
+FRigUnit_IsRecentlyRendered_Execute()
+{
+	if (const USkinnedMeshComponent* Component = MeshComponent)
+	{
+		Result = Component->bRecentlyRendered;
+	}
+	else
+	{
+		Result = false;
+	}
+}
diff --git a/Engine/Plugins/Experimental/AnimNext/Source/AnimNext/Private/Graph/RigUnit_IsRecentlyRendered.h b/Engine/Plugins/Experimental/AnimNext/Source/AnimNext/Private/Graph/RigUnit_IsRecentlyRendered.h
new file mode 100644
--- /dev/null
+++ b/Engine/Plugins/Experimental/AnimNext/Source/AnimNext/Private/Graph/RigUnit_IsRecentlyRendered.h
@@ -0,0 +1,26 @@
+// Copyright Epic Games, Inc. All Rights Reserved.
+
+#pragma once
+
+#include "AnimNextExecuteContext.h"
+#include "Graph/RigUnit_AnimNextBase.h"
+
+#include "RigUnit_IsRecentlyRendered.generated.h"
+
+/** Gets recently rendered flag from the input skinned mesh component */
+USTRUCT(meta = (DisplayName = "Is Recently Rendered", Category = "Mesh", NodeColor = "0, 1, 1", Keywords = "Mesh"))
+struct FRigUnit_IsRecentlyRendered : public FRigUnit_AnimNextBase
+{
+	GENERATED_BODY()
+
+	RIGVM_METHOD()
+	void Execute();
+
+	// Mesh to query
+	UPROPERTY(EditAnywhere, Category = "Graph", meta = (Input))
+	TObjectPtr<USkinnedMeshComponent> MeshComponent;
+
+	// The recently rendered flag of the input mesh
+	UPROPERTY(EditAnywhere, Category = "Graph", meta = (Output))
+	bool Result = false;
+};
