Add profiler scopes for game thread AnimNext tasks

diff --git a/Engine/Plugins/Experimental/AnimNext/Source/AnimNext/Private/Component/SkinnedMeshComponentExtensions.cpp b/Engine/Plugins/Experimental/AnimNext/Source/AnimNext/Private/Component/SkinnedMeshComponentExtensions.cpp
--- a/Engine/Plugins/Experimental/AnimNext/Source/AnimNext/Private/Component/SkinnedMeshComponentExtensions.cpp
+++ b/Engine/Plugins/Experimental/AnimNext/Source/AnimNext/Private/Component/SkinnedMeshComponentExtensions.cpp
@@ -38,6 +38,7 @@ void FSkinnedMeshComponentExtensions::CompleteAndDispatch(
 		{
 			FAnimNextModuleInstance::RunTaskOnGameThread([WeakComponent = TWeakObjectPtr<USkinnedMeshComponent>(InComponent)]()
 			{
+				SCOPED_NAMED_EVENT(AnimNext_SkinnedMesh_CompleteAndDispatch_GameThread, FColor::Orange);
 				USkinnedMeshComponent* Component = WeakComponent.Get();
 				if (Component == nullptr)
 				{
diff --git a/Engine/Plugins/Experimental/AnimNext/Source/AnimNext/Private/Module/ModuleTickFunction.cpp b/Engine/Plugins/Experimental/AnimNext/Source/AnimNext/Private/Module/ModuleTickFunction.cpp
--- a/Engine/Plugins/Experimental/AnimNext/Source/AnimNext/Private/Module/ModuleTickFunction.cpp
+++ b/Engine/Plugins/Experimental/AnimNext/Source/AnimNext/Private/Module/ModuleTickFunction.cpp
@@ -98,6 +98,7 @@ void FModuleEventTickFunction::EndTick(float DeltaTime)
 	{
 		FAnimNextModuleInstance::RunTaskOnGameThread([MainThreadActionEventList = MoveTemp(MainThreadActionEventList)]()
 			{
+				SCOPED_NAMED_EVENT(AnimNext_Module_EndTick_GameThread, FColor::Orange);
 				check(IsInGameThread());
 				for (const FAnimNextTraitEventPtr& Event : MainThreadActionEventList)
 				{
