Fix pose copy not working with a non-zero start index
Fix pose remap to use the target pose LOD where possible
Add pose debug fill to help track issues
Fix pose copy negative size

diff --git a/Engine/Plugins/Experimental/AnimNext/Source/AnimNext/Private/Graph/RemapPoseData.cpp b/Engine/Plugins/Experimental/AnimNext/Source/AnimNext/Private/Graph/RemapPoseData.cpp
--- a/Engine/Plugins/Experimental/AnimNext/Source/AnimNext/Private/Graph/RemapPoseData.cpp
+++ b/Engine/Plugins/Experimental/AnimNext/Source/AnimNext/Private/Graph/RemapPoseData.cpp
@@ -130,20 +130,17 @@ void FRemapPoseData::RemapPose(const UE::AnimNext::FLODPoseHeap& SourcePose,
 	SCOPE_CYCLE_COUNTER(STAT_AnimNext_RemapPose_Mesh2Mesh);
 
 	const int32 SourceLODLevel = SourcePose.LODLevel;
+	const int32 TargetLODLevel = TargetRefPose->GetSourceLODLevel();
 
-	const bool IsAdditive = false;
-	if (OutTargetPose.ShouldPrepareForLOD(*TargetRefPose, SourceLODLevel, SourcePose.IsAdditive()))
+	const bool bIsAdditive = SourcePose.IsAdditive();
+	if (OutTargetPose.ShouldPrepareForLOD(*TargetRefPose, TargetLODLevel, bIsAdditive))
 	{
-		OutTargetPose.PrepareForLOD(*TargetRefPose, SourceLODLevel, /*bSetRefPose=*/true, IsAdditive);
+		OutTargetPose.PrepareForLOD(*TargetRefPose, TargetLODLevel, /*bSetRefPose=*/true, bIsAdditive);
 	}
 
 	OutTargetPose.Flags = SourcePose.Flags;
 
-	// In case the source ref pose has fewer LOD levels available than the input/source pose, use the highest quality mapping available.
-	// This happens in case the source skeletal mesh does e.g. have 4 LODs while the target skeletal mesh only has 1. In this case we just create a single mapping
-	// for memory reasons as we'd store duplicated data otherwise and bones are guaranteed to be sorted.
-	const int32 TargetLODIndex = FMath::Clamp(SourceLODLevel, 0, SourceToTargetBoneIndexMapPerLOD[SourceLODLevel].Num() - 1);
-	const TArray<FRemappedBone>& SourceToTargetBoneIndexMap = SourceToTargetBoneIndexMapPerLOD[SourceLODLevel][TargetLODIndex].BoneIndexMap;
+	const TArray<FRemappedBone>& SourceToTargetBoneIndexMap = SourceToTargetBoneIndexMapPerLOD[SourceLODLevel][TargetLODLevel].BoneIndexMap;
 	const int32 NumToRemap = SourceToTargetBoneIndexMap.Num();
 
 #if DEFAULT_SOA && DEFAULT_SOA_VIEW
@@ -194,7 +191,7 @@ void FRemapPoseData::RemapPose(const UE::AnimNext::FLODPoseHeap& SourcePose,
 	// bone on the source skeleton for the target's root bone and sync the skeletons up from there by calculating the delta transform
 	// between the two and move it along with the source. This allows us to remap and attach skeletal meshes that only contain the
 	// bones it needs, like e.g. only skinned bones.
-	const int32 RootBoneOnSourceMapped = SourceToTargetBoneIndexMapPerLOD[SourceLODLevel][TargetLODIndex].TargetRootToSourceBoneIndex;
+	const int32 RootBoneOnSourceMapped = SourceToTargetBoneIndexMapPerLOD[SourceLODLevel][TargetLODLevel].TargetRootToSourceBoneIndex;
 	if (RootBoneOnSourceMapped != INDEX_NONE)
 	{
 		// TODO: Replace with the generic method once we have access to the modelspace transform via the LOD pose.
diff --git a/Engine/Plugins/Experimental/AnimNext/Source/AnimNext/Private/Graph/RigUnit_AnimNextWriteSkeletalMeshComponentPose.cpp b/Engine/Plugins/Experimental/AnimNext/Source/AnimNext/Private/Graph/RigUnit_AnimNextWriteSkeletalMeshComponentPose.cpp
--- a/Engine/Plugins/Experimental/AnimNext/Source/AnimNext/Private/Graph/RigUnit_AnimNextWriteSkeletalMeshComponentPose.cpp
+++ b/Engine/Plugins/Experimental/AnimNext/Source/AnimNext/Private/Graph/RigUnit_AnimNextWriteSkeletalMeshComponentPose.cpp
@@ -81,6 +81,10 @@ FRigUnit_AnimNextWriteSkeletalMeshComponentPose_Execute()
 	const FReferenceSkeleton& RefSkeleton = SkeletalMesh->GetRefSkeleton();
 	LocalSpaceTransforms.SetNumUninitialized(RefSkeleton.GetNum());
 
+#if UE_ENABLE_POSE_DEBUG_FILL
+	FMemory::Memset(LocalSpaceTransforms.GetData(), 0xCD, sizeof(FTransform) * LocalSpaceTransforms.Num());
+#endif
+
 	// Did we evaluate animation with lower quality than the visual mesh?
 	if (SkeletalMeshComponentLODLevel < PoseLODLevel)
 	{
diff --git a/Engine/Plugins/Experimental/AnimNext/Source/AnimNext/Private/TransformArrayOperations_SoA.cpp b/Engine/Plugins/Experimental/AnimNext/Source/AnimNext/Private/TransformArrayOperations_SoA.cpp
--- a/Engine/Plugins/Experimental/AnimNext/Source/AnimNext/Private/TransformArrayOperations_SoA.cpp
+++ b/Engine/Plugins/Experimental/AnimNext/Source/AnimNext/Private/TransformArrayOperations_SoA.cpp
@@ -86,12 +86,15 @@ namespace UE::AnimNext
 
 		// Our SoA buffer is contiguous
 		// Because the translations/scales have the same size (FVector), each entry is a fixed offset apart and we can use a single ptr/offset pair
-		const double* RESTRICT SrcTranslationPtr = reinterpret_cast<const double*>(Source.Translations.GetData()) + (StartIndex * 3);
-		double* RESTRICT DstTranslationPtr = reinterpret_cast<double*>(Dest.Translations.GetData()) + (StartIndex * 3);
+		const double* RESTRICT SrcTranslationPtr = reinterpret_cast<const double*>(Source.Translations.GetData());
+		double* RESTRICT DstTranslationPtr = reinterpret_cast<double*>(Dest.Translations.GetData());
 
 		const int64 SrcScaleOffset = reinterpret_cast<const double*>(Source.Scales3D.GetData()) - SrcTranslationPtr;
 		const int64 DstScaleOffset = reinterpret_cast<double*>(Dest.Scales3D.GetData()) - DstTranslationPtr;
 
+		SrcTranslationPtr += StartIndex * 3;
+		DstTranslationPtr += StartIndex * 3;
+
 		// Iterate until the last element
 		// Each translation/scale is 24 bytes (FVector) and everything is contiguous
 		// For all of these elements, we can use aligned 32-bytes load/stores for translation/scale
diff --git a/Engine/Plugins/Experimental/AnimNext/Source/AnimNext/Public/LODPose.h b/Engine/Plugins/Experimental/AnimNext/Source/AnimNext/Public/LODPose.h
--- a/Engine/Plugins/Experimental/AnimNext/Source/AnimNext/Public/LODPose.h
+++ b/Engine/Plugins/Experimental/AnimNext/Source/AnimNext/Public/LODPose.h
@@ -9,6 +9,9 @@
 #include "TransformArrayOperations.h"
 #include "LODPose.generated.h"
 
+// Handy define to set uninitialized pose buffers to 0xCDCDCDCD to help find issues
+#define UE_ENABLE_POSE_DEBUG_FILL 0
+
 namespace UE::AnimNext
 {
 
@@ -59,7 +62,7 @@ struct FLODPose
 			// The source pose is missing transforms as it is set to a lower-quality LOD level, initialize the missing bone transforms with the reference pose transforms.
 			const int32 NumSourceBones = SourcePose.LocalTransformsView.Num();
 			const int32 NumTargetBones = LocalTransformsView.Num();
-			const int32 NumAdditionalBones = NumSourceBones - NumTargetBones;
+			const int32 NumAdditionalBones = NumTargetBones - NumSourceBones;
 			SetRefPose(SourcePose.IsAdditive(), NumSourceBones, NumAdditionalBones);
 			CopyTransforms(LocalTransformsView, SourcePose.LocalTransformsView, /*StartIndex=*/0, NumSourceBones);
 		}
@@ -352,9 +355,18 @@ struct TLODPose : public FLODPose
 
 		Flags = (ELODPoseFlags)(bAdditive ? (Flags | ELODPoseFlags::Additive) : Flags & ELODPoseFlags::Additive);
 
-		if (bSetRefPose && NumTransforms > 0)
+		if (NumTransforms > 0)
 		{
-			SetRefPose(bAdditive);
+			if (bSetRefPose)
+			{
+				SetRefPose(bAdditive);
+			}
+#if UE_ENABLE_POSE_DEBUG_FILL
+			else
+			{
+				LocalTransforms.DebugFill();
+			}
+#endif
 		}
 	}
 };
diff --git a/Engine/Plugins/Experimental/AnimNext/Source/AnimNext/Public/TransformArray.h b/Engine/Plugins/Experimental/AnimNext/Source/AnimNext/Public/TransformArray.h
--- a/Engine/Plugins/Experimental/AnimNext/Source/AnimNext/Public/TransformArray.h
+++ b/Engine/Plugins/Experimental/AnimNext/Source/AnimNext/Public/TransformArray.h
@@ -204,6 +204,12 @@ struct TTransformArrayAoS
 		return true;
 	}
 
+	// Sets the whole array to 0xCDCDCDCD to aid debugging
+	void DebugFill()
+	{
+		FMemory::Memset(Transforms.GetData(), 0xCD, sizeof(FTransform) * Transforms.Num());
+	}
+
 private:
 #if ENABLE_NAN_DIAGNOSTIC
 	FORCEINLINE static void DiagnosticCheckNaN_Rotate(const FQuat& Rotation)
@@ -562,6 +568,12 @@ struct TTransformArraySoA
 		return true;
 	}
 
+	// Sets the whole array to 0xCDCDCDCD to aid debugging
+	void DebugFill()
+	{
+		FMemory::Memset(AllocatedMemory.GetData(), 0xCD, AllocatedMemory.Num());
+	}
+
 private:
 	TArray<uint8, AllocatorType> AllocatedMemory; // TODO : use the allocator directly
 
