[Optim] AnimNextModule trace optimizations:
1) Cache debug display name instead of searching every frame
2) Split module variable properties into description event and value event. Descriptions are large and take a long time to serialize, so we only send them once
3) Change crit section to RW lock in FAnimNextTrace

Windows editor
FAnimNextTrace::OutputAnimNextVariables Average 3.838ms -> 2.032ms

diff --git a/Engine/Plugins/Experimental/AnimNext/Source/AnimNext/Private/RewindDebugger/AnimNextTrace.cpp b/Engine/Plugins/Experimental/AnimNext/Source/AnimNext/Private/RewindDebugger/AnimNextTrace.cpp
--- a/Engine/Plugins/Experimental/AnimNext/Source/AnimNext/Private/RewindDebugger/AnimNextTrace.cpp
+++ b/Engine/Plugins/Experimental/AnimNext/Source/AnimNext/Private/RewindDebugger/AnimNextTrace.cpp
@@ -5,9 +5,11 @@
 #include "ObjectTrace.h"
 #include "Chaos/ParticleDirtyFlags.h"
 #include "DataInterface/AnimNextDataInterfaceInstance.h"
+#include "Misc/HashBuilder.h"
 #include "Module/AnimNextModuleInstance.h"
 #include "Serialization/MemoryWriter.h"
 #include "Serialization/ObjectAndNameAsStringProxyArchive.h"
+#include "Serialization/ObjectWriter.h"
 
 #if ANIMNEXT_TRACE_ENABLED
 
@@ -23,20 +25,47 @@ UE_TRACE_EVENT_BEGIN(AnimNext, InstanceVariables)
 	UE_TRACE_EVENT_FIELD(uint64, Cycle)
 	UE_TRACE_EVENT_FIELD(double, RecordingTime)
 	UE_TRACE_EVENT_FIELD(uint64, InstanceId)
+	UE_TRACE_EVENT_FIELD(uint32, VariableDescriptionHash)
 	UE_TRACE_EVENT_FIELD(uint8[], VariableData)
 UE_TRACE_EVENT_END()
 
+UE_TRACE_EVENT_BEGIN(AnimNext, InstanceVariableDescriptions)
+	UE_TRACE_EVENT_FIELD(uint32, VariableDescriptionHash)
+	UE_TRACE_EVENT_FIELD(uint8[], VariableDescriptionData)
+UE_TRACE_EVENT_END()
+
 UE_TRACE_CHANNEL_DEFINE(AnimNextChannel);
 
 namespace UE::AnimNext
 {
 
-FCriticalSection GTracedInstancesLock;
+FRWLock GTracedInstancesRWLock;
 TSet<uint64> GTracedInstances;
 
+FRWLock GTracedPropertiesRWLock;
+TMap<uint32, const UPropertyBag*> GTracedPropertyDescs;
+
+const FGuid FAnimNextTrace::CustomVersionGUID = FGuid(0x83E9BE13, 0xB1C845DC, 0x86C4D3E5, 0xE66CBE91);
+
+namespace FAnimNextTraceCustomVersion
+{
+enum CustomVersionIndex
+{
+	// Before any version changes were made in the plugin
+	FirstVersion = 0,
+
+	// -----<new versions can be added above this line>-------------------------------------------------
+	VersionPlusOne,
+	LatestVersion = VersionPlusOne - 1
+};
+}
+
+FCustomVersionRegistration GPropertyBagCustomVersion(FAnimNextTrace::CustomVersionGUID, FAnimNextTraceCustomVersion::FirstVersion, TEXT("FAnimNextTraceCustomVersion"));
+
 void FAnimNextTrace::Reset()
 {
 	GTracedInstances.Empty();
+	GTracedPropertyDescs.Empty();
 }
 
 void FAnimNextTrace::OutputAnimNextInstance(const FAnimNextDataInterfaceInstance* DataInterface, const UObject* OuterObject)
@@ -57,10 +86,13 @@ void FAnimNextTrace::OutputAnimNextInstance(const FAnimNextDataInterfaceInstance
 	bool bTrace = false;
 
 	{
-		FScopeLock Lock(&GTracedInstancesLock);
+		FRWScopeLock RWScopeLock(GTracedInstancesRWLock, SLT_ReadOnly);
 		
 		if (!GTracedInstances.Contains(InstanceId))
 		{
+			RWScopeLock.ReleaseReadOnlyLockAndAcquireWriteLock_USE_WITH_CAUTION();
+
+			// Note that we might double add here due to 2x people writing at the same time, but it doesn't matter
 			GTracedInstances.Add(InstanceId);
 			bTrace = true;
 		}
@@ -86,6 +118,22 @@ void FAnimNextTrace::OutputAnimNextInstance(const FAnimNextDataInterfaceInstance
 	}
 }
 
+uint32 GetPropertyDescHash(const TConstArrayView<FPropertyBagPropertyDesc>& Descs)
+{
+	FHashBuilder HashBuilder;
+	for (const FPropertyBagPropertyDesc& Desc : Descs)
+	{
+		HashBuilder.Append(Desc.Name);
+		HashBuilder.Append(Desc.ContainerTypes);
+		HashBuilder.Append(Desc.ID);
+		HashBuilder.Append(Desc.PropertyFlags);
+		HashBuilder.Append(Desc.ValueType);
+
+		// Q: should we care about metadata?
+	}
+	return HashBuilder.GetHash();
+}
+
 void FAnimNextTrace::OutputAnimNextVariables(const FAnimNextDataInterfaceInstance* Instance, const UObject* OuterObject)
 {
 	bool bChannelEnabled = UE_TRACE_CHANNELEXPR_IS_ENABLED(AnimNextChannel);
@@ -107,22 +155,55 @@ void FAnimNextTrace::OutputAnimNextVariables(const FAnimNextDataInterfaceInstanc
 	if (Variables.GetPropertyBagStruct())
 	{
 		TConstArrayView<FPropertyBagPropertyDesc> PropertyDescs = Variables.GetPropertyBagStruct()->GetPropertyDescs();
-		for(int i =0; i<PropertyDescs.Num(); i++)
+		uint32 PropertyDescHash = GetPropertyDescHash(PropertyDescs);
 		{
-			uint8* ValueMemory = Instance->GetMemoryForVariable(i, PropertyDescs[i].Name, PropertyDescs[i].CachedProperty);
-			PropertyDescs[i].CachedProperty->CopySingleValue(const_cast<uint8*>(Variables.GetValue().GetMemory()) +  PropertyDescs[i].CachedProperty->GetOffset_ForInternal(), ValueMemory);
+			FRWScopeLock RWScopeLock(GTracedPropertiesRWLock, SLT_ReadOnly);
+			TArray<uint8> ArchiveData;
+			ArchiveData.Reserve(1024 * 10);
+			
+			if (!GTracedPropertyDescs.Contains(PropertyDescHash))
+			{
+				// Upgrade to write lock
+				RWScopeLock.ReleaseReadOnlyLockAndAcquireWriteLock_USE_WITH_CAUTION();
+
+				// Check in case someone just serialized before we aquired the write lock
+				if (!GTracedPropertyDescs.Contains(PropertyDescHash))
+				{
+					GTracedPropertyDescs.Add(PropertyDescHash, Variables.GetPropertyBagStruct());
+					
+					// Serialize prop descs
+					TArray<FPropertyBagPropertyDesc, TInlineAllocator<64>> PropertyDescriptions;
+					PropertyDescriptions.Reserve(PropertyDescs.Num());
+					PropertyDescriptions = PropertyDescs;	// need to copy since we can't serialize an array view
+
+					FMemoryWriter WriterArchive(ArchiveData);
+					FObjectAndNameAsStringProxyArchive Archive(WriterArchive, /*bInLoadIfFindFails*/true);
+					Archive.UsingCustomVersion(CustomVersionGUID);
+					Archive.UsingCustomVersion(FPropertyBagCustomVersion::GUID);
+					Archive << PropertyDescriptions;
+
+					UE_TRACE_LOG(AnimNext, InstanceVariableDescriptions, AnimNextChannel)
+						<< InstanceVariableDescriptions.VariableDescriptionHash(PropertyDescHash)
+						<< InstanceVariableDescriptions.VariableDescriptionData(ArchiveData.GetData(), ArchiveData.Num());
+				}
+			}
+
+			ArchiveData.SetNum(0, EAllowShrinking::No);
+			FMemoryWriter WriterArchive(ArchiveData);
+			FObjectAndNameAsStringProxyArchive Archive(WriterArchive, /*bInLoadIfFindFails*/true);
+			UPropertyBag* PropertyBag = const_cast<UPropertyBag*>(Variables.GetPropertyBagStruct());
+			if (PropertyBag != nullptr)
+			{
+				PropertyBag->SerializeItem(Archive, Variables.GetMutableValue().GetMemory(), nullptr);
+			}
+
+			UE_TRACE_LOG(AnimNext, InstanceVariables, AnimNextChannel)
+					<< InstanceVariables.Cycle(FPlatformTime::Cycles64())
+					<< InstanceVariables.RecordingTime(FObjectTrace::GetWorldElapsedTime(OuterObject->GetWorld()))
+					<< InstanceVariables.InstanceId(InstanceId)
+					<< InstanceVariables.VariableDescriptionHash(PropertyDescHash)
+					<< InstanceVariables.VariableData(ArchiveData.GetData(), ArchiveData.Num());
 		}
-
-		TArray<uint8> ArchiveData;
-		FMemoryWriter WriterArchive(ArchiveData);
-		FObjectAndNameAsStringProxyArchive Archive(WriterArchive, /*bInLoadIfFindFails*/true);
-		const_cast<FInstancedPropertyBag&>(Variables).Serialize(Archive);
-
-		UE_TRACE_LOG(AnimNext, InstanceVariables, AnimNextChannel)
-				<< InstanceVariables.Cycle(FPlatformTime::Cycles64())
-				<< InstanceVariables.RecordingTime(FObjectTrace::GetWorldElapsedTime(OuterObject->GetWorld()))
-				<< InstanceVariables.InstanceId(InstanceId)
-				<< InstanceVariables.VariableData(ArchiveData.GetData(), ArchiveData.Num());
 	}
 	else
 	{
diff --git a/Engine/Plugins/Experimental/AnimNext/Source/AnimNext/Public/RewindDebugger/AnimNextTrace.h b/Engine/Plugins/Experimental/AnimNext/Source/AnimNext/Public/RewindDebugger/AnimNextTrace.h
--- a/Engine/Plugins/Experimental/AnimNext/Source/AnimNext/Public/RewindDebugger/AnimNextTrace.h
+++ b/Engine/Plugins/Experimental/AnimNext/Source/AnimNext/Public/RewindDebugger/AnimNextTrace.h
@@ -5,6 +5,7 @@
 #include "CoreTypes.h"
 #include "ObjectTrace.h"
 #include "CoreMinimal.h"
+#include "StructUtils/PropertyBag.h"
 #include "Trace/Trace.h"
 
 #define ANIMNEXT_TRACE_ENABLED (OBJECT_TRACE_ENABLED && !(UE_BUILD_SHIPPING || UE_BUILD_TEST))
@@ -19,9 +20,10 @@ struct FAnimNextGraphInstance;
 
 namespace UE::AnimNext
 {
-
 struct FAnimNextTrace
 {
+	ANIMNEXT_API static const FGuid CustomVersionGUID;
+	
 	ANIMNEXT_API static void Reset();
 	
 	ANIMNEXT_API static void OutputAnimNextInstance(const FAnimNextDataInterfaceInstance* DataInterface, const UObject* OuterObject);
diff --git a/Engine/Plugins/Experimental/AnimNext/Source/AnimNextEditor/Private/RewindDebugger/AnimNextAnalyzer.cpp b/Engine/Plugins/Experimental/AnimNext/Source/AnimNextEditor/Private/RewindDebugger/AnimNextAnalyzer.cpp
--- a/Engine/Plugins/Experimental/AnimNext/Source/AnimNextEditor/Private/RewindDebugger/AnimNextAnalyzer.cpp
+++ b/Engine/Plugins/Experimental/AnimNext/Source/AnimNextEditor/Private/RewindDebugger/AnimNextAnalyzer.cpp
@@ -22,6 +22,7 @@ void FAnimNextAnalyzer::OnAnalysisBegin(const FOnAnalysisContext& Context)
 
 	Builder.RouteEvent(RouteId_Module, "AnimNext", "Instance");
 	Builder.RouteEvent(RouteId_InstanceVariables, "AnimNext", "InstanceVariables");
+	Builder.RouteEvent(RouteId_InstanceVariableDescriptions, "AnimNext", "InstanceVariableDescriptions");
 }
 
 bool FAnimNextAnalyzer::OnEvent(uint16 RouteId, EStyle Style, const FOnEventContext& Context)
@@ -47,10 +48,18 @@ bool FAnimNextAnalyzer::OnEvent(uint16 RouteId, EStyle Style, const FOnEventCont
 			uint64 ModuleInstanceId = EventData.GetValue<uint64>("InstanceId");
 			uint64 Cycle = EventData.GetValue<uint64>("Cycle");
 			double RecordingTime = EventData.GetValue<double>("RecordingTime");
-
+			uint32 VariableDescHash = EventData.GetValue<uint32>("VariableDescriptionHash");
 			TArrayView<const uint8> VariableData = EventData.GetArrayView<uint8>("VariableData");
 
-			Provider.AppendVariables(Context.EventTime.AsSeconds(Cycle), RecordingTime, ModuleInstanceId, VariableData);
+			Provider.AppendVariables(Context.EventTime.AsSeconds(Cycle), RecordingTime, ModuleInstanceId, VariableDescHash, VariableData);
+			break;
+		}
+		case RouteId_InstanceVariableDescriptions:
+		{
+			uint32 VariableDescHash = EventData.GetValue<uint32>("VariableDescriptionHash");
+			TArrayView<const uint8> VariableDescData = EventData.GetArrayView<uint8>("VariableDescriptionData");
+
+			Provider.AppendVariableDescriptions(VariableDescHash, VariableDescData);
 			break;
 		}
 	}
diff --git a/Engine/Plugins/Experimental/AnimNext/Source/AnimNextEditor/Private/RewindDebugger/AnimNextAnalyzer.h b/Engine/Plugins/Experimental/AnimNext/Source/AnimNextEditor/Private/RewindDebugger/AnimNextAnalyzer.h
--- a/Engine/Plugins/Experimental/AnimNext/Source/AnimNextEditor/Private/RewindDebugger/AnimNextAnalyzer.h
+++ b/Engine/Plugins/Experimental/AnimNext/Source/AnimNextEditor/Private/RewindDebugger/AnimNextAnalyzer.h
@@ -21,6 +21,7 @@ private:
 	{
 		RouteId_Module,
 		RouteId_InstanceVariables,
+		RouteId_InstanceVariableDescriptions
 	};
 
 	TraceServices::IAnalysisSession& Session;
diff --git a/Engine/Plugins/Experimental/AnimNext/Source/AnimNextEditor/Private/RewindDebugger/AnimNextModuleTrack.cpp b/Engine/Plugins/Experimental/AnimNext/Source/AnimNextEditor/Private/RewindDebugger/AnimNextModuleTrack.cpp
--- a/Engine/Plugins/Experimental/AnimNext/Source/AnimNextEditor/Private/RewindDebugger/AnimNextModuleTrack.cpp
+++ b/Engine/Plugins/Experimental/AnimNext/Source/AnimNextEditor/Private/RewindDebugger/AnimNextModuleTrack.cpp
@@ -9,9 +9,12 @@
 #include "Subsystems/AssetEditorSubsystem.h"
 #include "Modules/ModuleManager.h"
 #include "PropertyEditorModule.h"
+#include "RewindDebugger/AnimNextTrace.h"
 #include "Serialization/ObjectAndNameAsStringProxyArchive.h"
 #include "Serialization/ObjectReader.h"
-
+#include "StructUtils/PropertyBag.h"
+#include "Serialization/MemoryReader.h"
+#include "StructUtils/UserDefinedStruct.h"
 
 #define LOCTEXT_NAMESPACE "AnimNextModuleTrack"
 
@@ -29,20 +32,27 @@ static const FName AnimNextModulesName("AnimNextModules");
 	
 FText FAnimNextModuleTrack::GetDisplayNameInternal() const
 {
-	TraceServices::FAnalysisSessionReadScope SessionReadScope(*IRewindDebugger::Instance()->GetAnalysisSession());
-	const FAnimNextProvider* AnimNextProvider = IRewindDebugger::Instance()->GetAnalysisSession()->ReadProvider<FAnimNextProvider>("AnimNextProvider");
-	const IGameplayProvider* GameplayProvider = IRewindDebugger::Instance()->GetAnalysisSession()->ReadProvider<IGameplayProvider>("GameplayProvider");
-
- 	if (InstanceId != 0)
+ 	if (DisplayNameCache.IsEmpty())
  	{
- 		if (const FDataInterfaceData* Data = AnimNextProvider->GetDataInterfaceData(InstanceId))
+ 		TraceServices::FAnalysisSessionReadScope SessionReadScope(*IRewindDebugger::Instance()->GetAnalysisSession());
+ 		const FAnimNextProvider* AnimNextProvider = IRewindDebugger::Instance()->GetAnalysisSession()->ReadProvider<FAnimNextProvider>("AnimNextProvider");
+ 		const IGameplayProvider* GameplayProvider = IRewindDebugger::Instance()->GetAnalysisSession()->ReadProvider<IGameplayProvider>("GameplayProvider");
+
+ 		if (InstanceId != 0)
  		{
-			const FObjectInfo& ModuleInfo = GameplayProvider->GetObjectInfo(Data->AssetId);
-			return FText::FromString(ModuleInfo.Name);
-		}
- 	}
+ 			if (const FDataInterfaceData* Data = AnimNextProvider->GetDataInterfaceData(InstanceId))
+ 			{
+ 				const FObjectInfo& ModuleInfo = GameplayProvider->GetObjectInfo(Data->AssetId);
+ 				DisplayNameCache = FText::FromString(ModuleInfo.Name);
+ 				return DisplayNameCache;
+ 			}
+ 		}
 	
- 	return NSLOCTEXT("RewindDebugger", "AnimNextModuleTrackName", "Module");
+		// Don't cache this since it is a placeholder name, will be replaced once the real one is sent
+ 		return NSLOCTEXT("RewindDebugger", "AnimNextModuleTrackName", "Module");
+ 	}
+
+ 	return DisplayNameCache;
 }
 
 void FAnimNextModuleTrackCreator::GetTrackTypesInternal(TArray<RewindDebugger::FRewindDebuggerTrackType>& Types) const
@@ -145,17 +155,42 @@ bool FAnimNextModuleTrack::UpdateInternal()
 					const TraceServices::IFrameProvider& FramesProvider = TraceServices::ReadFrameProvider(*AnalysisSession);
 					TraceServices::FAnalysisSessionReadScope SessionReadScope(*AnalysisSession);
 
-					DetailsObject->Properties.Reset();
-
 					TraceServices::FFrame MarkerFrame;
 					if(FramesProvider.GetFrameFromTime(ETraceFrameType::TraceFrameType_Game, CurrentScrubTime, MarkerFrame))
 					{
-						Data->VariablesTimeline.EnumerateEvents(MarkerFrame.StartTime, MarkerFrame.EndTime, [DetailsObject](double InStartTime, double InEndTime, uint32 InDepth, const TArray<uint8>& VariableData)
-						{				
-							FMemoryReader Reader(VariableData);
-							FObjectAndNameAsStringProxyArchive ReaderProxy(Reader, /*bInLoadIfFindFails*/true);
-							DetailsObject->Properties.Serialize(ReaderProxy);
-						
+						Data->VariablesTimeline.EnumerateEvents(MarkerFrame.StartTime, MarkerFrame.EndTime, [AnimNextProvider, DetailsObject, this](double InStartTime, double InEndTime, uint32 InDepth, const FPropertyVariableData& VariableListData)
+						{
+							// First look up property description
+							const FPropertyDescriptionData* DescriptionData = AnimNextProvider->GetPropertyDescriptionData(VariableListData.DescriptionHash);
+							if (DescriptionData == nullptr)
+							{
+								// Can't do anything without a description
+								return TraceServices::EEventEnumerate::Continue;
+							}
+							
+							{
+								// Load the property descriptions
+								FMemoryReader Reader(DescriptionData->Data);
+								FObjectAndNameAsStringProxyArchive ReaderProxy(Reader, /*bInLoadIfFindFails*/true);
+								ReaderProxy.UsingCustomVersion(AnimNext::FAnimNextTrace::CustomVersionGUID);
+								ReaderProxy << PropertyDescriptions;
+
+								DetailsObject->Properties.Reset();
+								DetailsObject->Properties.AddProperties(PropertyDescriptions);
+							}
+
+							{
+								// Load the property values
+								FMemoryReader Reader(VariableListData.ValueData);
+								FObjectAndNameAsStringProxyArchive ReaderProxy(Reader, /*bInLoadIfFindFails*/true);
+
+								UPropertyBag* PropertyBag = const_cast<UPropertyBag*>(DetailsObject->Properties.GetPropertyBagStruct());
+								if (PropertyBag != nullptr)
+								{
+									PropertyBag->SerializeItem(ReaderProxy, DetailsObject->Properties.GetMutableValue().GetMemory(), nullptr);
+								}
+							}
+
 							return TraceServices::EEventEnumerate::Stop;
 						});
 					}
diff --git a/Engine/Plugins/Experimental/AnimNext/Source/AnimNextEditor/Private/RewindDebugger/AnimNextModuleTrack.h b/Engine/Plugins/Experimental/AnimNext/Source/AnimNextEditor/Private/RewindDebugger/AnimNextModuleTrack.h
--- a/Engine/Plugins/Experimental/AnimNext/Source/AnimNextEditor/Private/RewindDebugger/AnimNextModuleTrack.h
+++ b/Engine/Plugins/Experimental/AnimNext/Source/AnimNextEditor/Private/RewindDebugger/AnimNextModuleTrack.h
@@ -54,8 +54,10 @@ private:
 	double PreviousScrubTime = -1.0;
 	TWeakObjectPtr<UPropertyBagDetailsObject> DetailsObjectWeakPtr;
 	TSharedPtr<SEventTimelineView::FTimelineEventData> ExistenceRange;
+	mutable FText DisplayNameCache;
 
 	TArray<TSharedPtr<FAnimNextModuleTrack>> Children;
+	TArray<FPropertyBagPropertyDesc> PropertyDescriptions;
 };
 
 
diff --git a/Engine/Plugins/Experimental/AnimNext/Source/AnimNextEditor/Private/RewindDebugger/AnimNextProvider.cpp b/Engine/Plugins/Experimental/AnimNext/Source/AnimNextEditor/Private/RewindDebugger/AnimNextProvider.cpp
--- a/Engine/Plugins/Experimental/AnimNext/Source/AnimNextEditor/Private/RewindDebugger/AnimNextProvider.cpp
+++ b/Engine/Plugins/Experimental/AnimNext/Source/AnimNextEditor/Private/RewindDebugger/AnimNextProvider.cpp
@@ -28,9 +28,12 @@ void FAnimNextProvider::AppendInstance(uint64 InstanceId, uint64 HostInstanceId,
 	Data->AssetId = AssetId;
 
 	DataInterfaceData.Add(InstanceId, Data);
+	
+	TArray<TSharedRef<FDataInterfaceData>>& ChildList = HostToChildDataMap.FindOrAdd(HostInstanceId);
+	ChildList.AddUnique(Data);
 }
 
-void FAnimNextProvider::AppendVariables(double ProfileTime, double RecordingTime, uint64 DataInterfaceId, const TArrayView<const uint8>& VariableData)
+void FAnimNextProvider::AppendVariables(double ProfileTime, double RecordingTime, uint64 DataInterfaceId, uint32 PropertyDescriptionHash, const TArrayView<const uint8>& VariableData)
 {
 	Session.WriteAccessCheck();
 	
@@ -38,7 +41,7 @@ void FAnimNextProvider::AppendVariables(double ProfileTime, double RecordingTime
 	{
 		if (!VariableData.IsEmpty())
 		{
-			(*Data)->VariablesTimeline.AppendEvent(ProfileTime, TArray<uint8>(VariableData));
+			(*Data)->VariablesTimeline.AppendEvent(ProfileTime, FPropertyVariableData(PropertyDescriptionHash, VariableData));
 		}
 		(*Data)->EndTime = RecordingTime;
 		if ((*Data)->StartTime < 0)
@@ -48,6 +51,20 @@ void FAnimNextProvider::AppendVariables(double ProfileTime, double RecordingTime
 	}
 }
 
+void FAnimNextProvider::AppendVariableDescriptions(uint32 PropertyDescriptionHash, const TArrayView<const uint8>& VariableDescriptionData)
+{
+	Session.WriteAccessCheck();
+	
+	if (HashToPropertyDescriptionDataMap.Contains(PropertyDescriptionHash) == false)
+	{
+		if (!VariableDescriptionData.IsEmpty())
+		{
+			HashToPropertyDescriptionDataMap.Add(PropertyDescriptionHash, FPropertyDescriptionData(VariableDescriptionData));
+		}
+	}
+}
+
+
 bool FAnimNextProvider::GetModuleId(uint64 ComponentId, uint64& OutModuleId) const
 {
 	Session.ReadAccessCheck();
@@ -74,15 +91,15 @@ void FAnimNextProvider::EnumerateChildInstances(uint64 InstanceId, TFunctionRef<
 {
 	Session.ReadAccessCheck();
 
-	for (auto& Data : DataInterfaceData)
+	if (const TArray<TSharedRef<FDataInterfaceData>>* Children = HostToChildDataMap.Find(InstanceId))
 	{
-		const TSharedPtr<FDataInterfaceData>& Instance = Data.Value;
-		
-		if (Instance->HostInstanceId == InstanceId)
+		for (auto Data : *Children)
 		{
-			Callback(*Instance);
+			check (Data->HostInstanceId == InstanceId);
+			Callback(*Data);
 		}
 	}
+
 }
 
 #undef LOCTEXT_NAMESPACE
diff --git a/Engine/Plugins/Experimental/AnimNext/Source/AnimNextEditor/Private/RewindDebugger/AnimNextProvider.h b/Engine/Plugins/Experimental/AnimNext/Source/AnimNextEditor/Private/RewindDebugger/AnimNextProvider.h
--- a/Engine/Plugins/Experimental/AnimNext/Source/AnimNextEditor/Private/RewindDebugger/AnimNextProvider.h
+++ b/Engine/Plugins/Experimental/AnimNext/Source/AnimNextEditor/Private/RewindDebugger/AnimNextProvider.h
@@ -11,6 +11,21 @@
 
 namespace TraceServices { class IAnalysisSession; }
 
+struct FPropertyDescriptionData
+{
+	FPropertyDescriptionData(TConstArrayView<uint8> DataView) : Data(DataView) {}
+	TArray<uint8> Data;
+};
+
+struct FPropertyVariableData
+{
+	FPropertyVariableData(uint32 DescriptionHash, TConstArrayView<uint8> Data)
+		: DescriptionHash(DescriptionHash) , ValueData(Data) {}
+	
+	uint32 DescriptionHash;
+	TArray<uint8> ValueData;
+};
+
 struct FDataInterfaceData
 {
 	FDataInterfaceData(TraceServices::IAnalysisSession& Session)
@@ -24,7 +39,7 @@ struct FDataInterfaceData
 	{
 	}
 
-	TraceServices::TPointTimeline<TArray<uint8>> VariablesTimeline;
+	TraceServices::TPointTimeline<FPropertyVariableData> VariablesTimeline;
 	uint64 InstanceId;
 	uint64 HostInstanceId;
 	uint64 OuterObjectId;
@@ -41,10 +56,12 @@ public:
 	FAnimNextProvider(TraceServices::IAnalysisSession& InSession);
 
 	void AppendInstance(uint64 InstanceId, uint64 HostInstanceId,  uint64 AssetId, uint64 OuterObjectId);
-	void AppendVariables(double ProfileTime, double RecordingTime, uint64 InModuleInstanceId, const TArrayView<const uint8>& VariableData);
+	void AppendVariables(double ProfileTime, double RecordingTime, uint64 InModuleInstanceId, uint32 PropertyDescriptionHash, const TArrayView<const uint8>& VariableData);
+	void AppendVariableDescriptions(uint32 PropertyDescriptionHash, const TArrayView<const uint8>& VariableDescriptionData);
 
 	bool GetModuleId(uint64 ComponentId, uint64& OutModuleId) const;
 	const FDataInterfaceData* GetDataInterfaceData(uint64 DataInterfaceId) const;
+	const FPropertyDescriptionData* GetPropertyDescriptionData(uint32 PropertyDescriptionHash) const { return HashToPropertyDescriptionDataMap.Find(PropertyDescriptionHash);}
 
 	void EnumerateChildInstances(uint64 InstanceId, TFunctionRef<void(const FDataInterfaceData&)> Callback) const;
 
@@ -53,4 +70,6 @@ private:
 
 	TMap<uint64, uint64> ComponentIdToModuleId;
 	TMap<uint64, TSharedRef<FDataInterfaceData>> DataInterfaceData;
+	TMap<uint64, TArray<TSharedRef<FDataInterfaceData>>> HostToChildDataMap;
+	TMap<uint32, FPropertyDescriptionData> HashToPropertyDescriptionDataMap;
 };
diff --git a/Engine/Source/Runtime/CoreUObject/Private/StructUtils/PropertyBag.cpp b/Engine/Source/Runtime/CoreUObject/Private/StructUtils/PropertyBag.cpp
--- a/Engine/Source/Runtime/CoreUObject/Private/StructUtils/PropertyBag.cpp
+++ b/Engine/Source/Runtime/CoreUObject/Private/StructUtils/PropertyBag.cpp
@@ -28,7 +28,7 @@ static TAutoConsoleVariable<bool> CVarValidatePropertyBagDescriptionEquality(
 
 #include UE_INLINE_GENERATED_CPP_BY_NAME(PropertyBag)
 
-struct FPropertyBagCustomVersion
+namespace FPropertyBagCustomVersion
 {
 	enum Type
 	{
@@ -44,12 +44,6 @@ struct FPropertyBagCustomVersion
 		VersionPlusOne,
 		LatestVersion = VersionPlusOne - 1
 	};
-
-	// The GUID for this custom version number
-	const static FGuid GUID;
-
-private:
-	FPropertyBagCustomVersion() {}
 };
 
 
@@ -1665,7 +1659,7 @@ FPropertyBagPropertyDesc::FPropertyBagPropertyDesc(const FName InName, const FPr
 #endif
 }
 
-static FArchive& operator<<(FArchive& Ar, FPropertyBagPropertyDesc& Bag)
+FArchive& operator<<(FArchive& Ar, FPropertyBagPropertyDesc& Bag)
 {
 	Ar << Bag.ValueTypeObject;
 	Ar << Bag.ID;
diff --git a/Engine/Source/Runtime/CoreUObject/Public/StructUtils/PropertyBag.h b/Engine/Source/Runtime/CoreUObject/Public/StructUtils/PropertyBag.h
--- a/Engine/Source/Runtime/CoreUObject/Public/StructUtils/PropertyBag.h
+++ b/Engine/Source/Runtime/CoreUObject/Public/StructUtils/PropertyBag.h
@@ -48,6 +48,11 @@ enum class EPropertyBagContainerType : uint8
 	Count UMETA(Hidden)
 };
 
+namespace FPropertyBagCustomVersion
+{
+	COREUOBJECT_API extern const FGuid GUID;
+}
+
 /** Helper to manage container types, with nested container support. */
 USTRUCT()
 struct FPropertyBagContainerTypes
@@ -294,6 +299,14 @@ struct FPropertyBagPropertyDesc
 	/** @return true if the property bag's unique IDs match */
 	UE_API bool operator==(const FPropertyBagPropertyDesc& OtherDesc) const;
 
+	/**
+	 * Serializes a PropertyBagPropertyDesc from or into an archive.
+	 *
+	 * @param Ar The archive to serialize from or into.
+	 * @param FPropertyBagPropertyDesc The PropertyBagPropertyDesc to serialize.
+	 */
+	 UE_API friend FArchive& operator<<(FArchive& Ar, FPropertyBagPropertyDesc& Desc);
+
 	/** Pointer to object that defines the Enum, Struct, or Class. */
 	UPROPERTY(EditAnywhere, Category="Default")
 	TObjectPtr<const UObject> ValueTypeObject = nullptr;
