Buildfix AnimNextModuleTrack.cpp EngineTestEditorCooked

diff --git a/Engine/Plugins/Experimental/AnimNext/Source/AnimNextEditor/Private/RewindDebugger/AnimNextModuleTrack.cpp b/Engine/Plugins/Experimental/AnimNext/Source/AnimNextEditor/Private/RewindDebugger/AnimNextModuleTrack.cpp
--- a/Engine/Plugins/Experimental/AnimNext/Source/AnimNextEditor/Private/RewindDebugger/AnimNextModuleTrack.cpp
+++ b/Engine/Plugins/Experimental/AnimNext/Source/AnimNextEditor/Private/RewindDebugger/AnimNextModuleTrack.cpp
@@ -2,6 +2,7 @@
 
 #include "AnimNextModuleTrack.h"
 
+#if ANIMNEXT_TRACE_ENABLED
 #include "AnimNextProvider.h"
 #include "Editor.h"
 #include "IGameplayProvider.h"
@@ -9,12 +10,9 @@
 #include "Subsystems/AssetEditorSubsystem.h"
 #include "Modules/ModuleManager.h"
 #include "PropertyEditorModule.h"
-#include "RewindDebugger/AnimNextTrace.h"
 #include "Serialization/ObjectAndNameAsStringProxyArchive.h"
-#include "Serialization/ObjectReader.h"
 #include "StructUtils/PropertyBag.h"
 #include "Serialization/MemoryReader.h"
-#include "StructUtils/UserDefinedStruct.h"
 
 #define LOCTEXT_NAMESPACE "AnimNextModuleTrack"
 
@@ -286,4 +284,6 @@ bool FAnimNextModuleTrackCreator::HasDebugInfoInternal(uint64 ObjectId) const
 
 }
 
-#undef LOCTEXT_NAMESPACE
\ No newline at end of file
+#undef LOCTEXT_NAMESPACE
+
+#endif // ANIMNEXT_TRACE_ENABLED
diff --git a/Engine/Plugins/Experimental/AnimNext/Source/AnimNextEditor/Private/RewindDebugger/AnimNextModuleTrack.h b/Engine/Plugins/Experimental/AnimNext/Source/AnimNextEditor/Private/RewindDebugger/AnimNextModuleTrack.h
--- a/Engine/Plugins/Experimental/AnimNext/Source/AnimNextEditor/Private/RewindDebugger/AnimNextModuleTrack.h
+++ b/Engine/Plugins/Experimental/AnimNext/Source/AnimNextEditor/Private/RewindDebugger/AnimNextModuleTrack.h
@@ -1,13 +1,16 @@
 // Copyright Epic Games, Inc. All Rights Reserved.
 #pragma once
 
+#include "RewindDebugger/AnimNextTrace.h"
+
+#if ANIMNEXT_TRACE_ENABLED
 #include "IRewindDebuggerTrackCreator.h"
 #include "RewindDebuggerTrack.h"
 #include "Textures/SlateIcon.h"
 #include "SEventTimelineView.h"
 #include "StructUtils/PropertyBag.h"
 #include "IDetailsView.h"
-#include "SEventTimelineView.h"
+#endif // ANIMNEXT_TRACE_ENABLED 
 
 #include "AnimNextModuleTrack.generated.h"
 
@@ -21,6 +24,7 @@ public:
 	FInstancedPropertyBag Properties;
 };
 
+#if ANIMNEXT_TRACE_ENABLED
 namespace UE::AnimNextEditor
 {
 	
@@ -71,4 +75,5 @@ private:
 	virtual bool HasDebugInfoInternal(uint64 ObjectId) const override;
 };
 
-}
\ No newline at end of file
+}
+#endif // ANIMNEXT_TRACE_ENABLED
