Compile fixes for anim next trace in editor shipping config

diff --git a/Engine/Plugins/Experimental/AnimNext/Source/AnimNextEditor/Private/AnimNextEditorModule.cpp b/Engine/Plugins/Experimental/AnimNext/Source/AnimNextEditor/Private/AnimNextEditorModule.cpp
--- a/Engine/Plugins/Experimental/AnimNext/Source/AnimNextEditor/Private/AnimNextEditorModule.cpp
+++ b/Engine/Plugins/Experimental/AnimNext/Source/AnimNextEditor/Private/AnimNextEditorModule.cpp
@@ -79,12 +79,14 @@
 
 #define LOCTEXT_NAMESPACE "AnimNextEditorModule"
 
+#if ANIMNEXT_TRACE_ENABLED
 namespace
 {
 	FRewindDebuggerAnimNext GRewindDebuggerAnimNext;
 	UE::AnimNextEditor::FAnimNextModuleTrackCreator GAnimNextModulesTrackCreator;
 	FAnimNextTraceModule GAnimNextTraceModule;
 }
+#endif
 
 namespace UE::AnimNext::Editor
 {
@@ -102,10 +104,12 @@ void FAnimNextEditorModule::StartupModule()
 		LOCTEXT("SettingsDescription", "Customize AnimNext Settings."),
 		GetMutableDefault<UAnimNextConfig>()
 	);
-	
+
+#if ANIMNEXT_TRACE_ENABLED
 	IModularFeatures::Get().RegisterModularFeature(IRewindDebuggerExtension::ModularFeatureName, &GRewindDebuggerAnimNext);
 	IModularFeatures::Get().RegisterModularFeature(RewindDebugger::IRewindDebuggerTrackCreator::ModularFeatureName, &GAnimNextModulesTrackCreator);
 	IModularFeatures::Get().RegisterModularFeature(TraceServices::ModuleFeatureName, &GAnimNextTraceModule);
+#endif // ANIMNEXT_TRACE_ENABLED
 
 	FPropertyEditorModule& PropertyModule = FModuleManager::LoadModuleChecked<FPropertyEditorModule>("PropertyEditor");
 
@@ -421,10 +425,12 @@ void FAnimNextEditorModule::ShutdownModule()
 		PersonaModule.UnregisterNotifyHostAsset(FTopLevelAssetPath(TEXT("/Script/AnimNext.AnimNextRigVMAsset")));
 	}
 	UnregisterAssetCompilationHandler(FTopLevelAssetPath(TEXT("/Script/AnimNext.AnimNextModule")));
-	
+
+#if ANIMNEXT_TRACE_ENABLED
 	IModularFeatures::Get().UnregisterModularFeature(IRewindDebuggerExtension::ModularFeatureName, &GRewindDebuggerAnimNext);
 	IModularFeatures::Get().UnregisterModularFeature(RewindDebugger::IRewindDebuggerTrackCreator::ModularFeatureName, &GAnimNextModulesTrackCreator);
 	IModularFeatures::Get().UnregisterModularFeature(TraceServices::ModuleFeatureName, &GAnimNextTraceModule);
+#endif // ANIMNEXT_TRACE_ENABLED
 
 	FEdGraphUtilities::UnregisterVisualPinFactory(GraphPanelPinFactory);
 }
diff --git a/Engine/Plugins/Experimental/AnimNext/Source/AnimNextEditor/Private/RewindDebugger/AnimNextAnalyzer.cpp b/Engine/Plugins/Experimental/AnimNext/Source/AnimNextEditor/Private/RewindDebugger/AnimNextAnalyzer.cpp
--- a/Engine/Plugins/Experimental/AnimNext/Source/AnimNextEditor/Private/RewindDebugger/AnimNextAnalyzer.cpp
+++ b/Engine/Plugins/Experimental/AnimNext/Source/AnimNextEditor/Private/RewindDebugger/AnimNextAnalyzer.cpp
@@ -2,6 +2,8 @@
 
 #include "AnimNextAnalyzer.h"
 
+#if ANIMNEXT_TRACE_ENABLED
+
 #include "HAL/LowLevelMemTracker.h"
 #include "Serialization/MemoryReader.h"
 #include "TraceServices/Model/AnalysisSession.h"
@@ -66,3 +68,4 @@ bool FAnimNextAnalyzer::OnEvent(uint16 RouteId, EStyle Style, const FOnEventCont
 
 	return true;
 }
+#endif // ANIMNEXT_TRACE_ENABLED
diff --git a/Engine/Plugins/Experimental/AnimNext/Source/AnimNextEditor/Private/RewindDebugger/AnimNextAnalyzer.h b/Engine/Plugins/Experimental/AnimNext/Source/AnimNextEditor/Private/RewindDebugger/AnimNextAnalyzer.h
--- a/Engine/Plugins/Experimental/AnimNext/Source/AnimNextEditor/Private/RewindDebugger/AnimNextAnalyzer.h
+++ b/Engine/Plugins/Experimental/AnimNext/Source/AnimNextEditor/Private/RewindDebugger/AnimNextAnalyzer.h
@@ -2,6 +2,9 @@
 
 #pragma once
 
+#include "RewindDebugger/AnimNextTrace.h"
+
+#if ANIMNEXT_TRACE_ENABLED
 #include "Trace/Analyzer.h"
 
 class FAnimNextProvider;
@@ -27,3 +30,4 @@ private:
 	TraceServices::IAnalysisSession& Session;
 	FAnimNextProvider& Provider;
 };
+#endif // ANIMNEXT_TRACE_ENABLED
diff --git a/Engine/Plugins/Experimental/AnimNext/Source/AnimNextEditor/Private/RewindDebugger/AnimNextProvider.cpp b/Engine/Plugins/Experimental/AnimNext/Source/AnimNextEditor/Private/RewindDebugger/AnimNextProvider.cpp
--- a/Engine/Plugins/Experimental/AnimNext/Source/AnimNextEditor/Private/RewindDebugger/AnimNextProvider.cpp
+++ b/Engine/Plugins/Experimental/AnimNext/Source/AnimNextEditor/Private/RewindDebugger/AnimNextProvider.cpp
@@ -1,6 +1,7 @@
 // Copyright Epic Games, Inc. All Rights Reserved.
 
 #include "AnimNextProvider.h"
+#if ANIMNEXT_TRACE_ENABLED
 #include "ObjectTrace.h"
 
 FName FAnimNextProvider::ProviderName("AnimNextProvider");
@@ -103,3 +104,4 @@ void FAnimNextProvider::EnumerateChildInstances(uint64 InstanceId, TFunctionRef<
 }
 
 #undef LOCTEXT_NAMESPACE
+#endif // ANIMNEXT_TRACE_ENABLED
diff --git a/Engine/Plugins/Experimental/AnimNext/Source/AnimNextEditor/Private/RewindDebugger/AnimNextProvider.h b/Engine/Plugins/Experimental/AnimNext/Source/AnimNextEditor/Private/RewindDebugger/AnimNextProvider.h
--- a/Engine/Plugins/Experimental/AnimNext/Source/AnimNextEditor/Private/RewindDebugger/AnimNextProvider.h
+++ b/Engine/Plugins/Experimental/AnimNext/Source/AnimNextEditor/Private/RewindDebugger/AnimNextProvider.h
@@ -2,6 +2,9 @@
 
 #pragma once
 
+#include "RewindDebugger/AnimNextTrace.h"
+
+#if ANIMNEXT_TRACE_ENABLED
 #include "Engine/EngineBaseTypes.h"
 #include "Model/PointTimeline.h"
 #include "Containers/Map.h"
@@ -73,3 +76,4 @@ private:
 	TMap<uint64, TArray<TSharedRef<FDataInterfaceData>>> HostToChildDataMap;
 	TMap<uint32, FPropertyDescriptionData> HashToPropertyDescriptionDataMap;
 };
+#endif //ANIMNEXT_TRACE_ENABLED
\ No newline at end of file
diff --git a/Engine/Plugins/Experimental/AnimNext/Source/AnimNextEditor/Private/RewindDebugger/AnimNextTraceModule.cpp b/Engine/Plugins/Experimental/AnimNext/Source/AnimNextEditor/Private/RewindDebugger/AnimNextTraceModule.cpp
--- a/Engine/Plugins/Experimental/AnimNext/Source/AnimNextEditor/Private/RewindDebugger/AnimNextTraceModule.cpp
+++ b/Engine/Plugins/Experimental/AnimNext/Source/AnimNextEditor/Private/RewindDebugger/AnimNextTraceModule.cpp
@@ -1,6 +1,8 @@
 // Copyright Epic Games, Inc. All Rights Reserved.
 
 #include "AnimNextTraceModule.h"
+
+#if ANIMNEXT_TRACE_ENABLED
 #include "AnimNextProvider.h"
 #include "AnimNextAnalyzer.h"
 
@@ -29,4 +31,4 @@ void FAnimNextTraceModule::GenerateReports(const TraceServices::IAnalysisSession
 {
 
 }
-
+#endif //ANIMNEXT_TRACE_ENABLED
diff --git a/Engine/Plugins/Experimental/AnimNext/Source/AnimNextEditor/Private/RewindDebugger/AnimNextTraceModule.h b/Engine/Plugins/Experimental/AnimNext/Source/AnimNextEditor/Private/RewindDebugger/AnimNextTraceModule.h
--- a/Engine/Plugins/Experimental/AnimNext/Source/AnimNextEditor/Private/RewindDebugger/AnimNextTraceModule.h
+++ b/Engine/Plugins/Experimental/AnimNext/Source/AnimNextEditor/Private/RewindDebugger/AnimNextTraceModule.h
@@ -2,6 +2,9 @@
 
 #pragma once
 
+#include "RewindDebugger/AnimNextTrace.h"
+
+#if ANIMNEXT_TRACE_ENABLED
 #include "CoreMinimal.h"
 #include "TraceServices/ModuleService.h"
 
@@ -20,4 +23,5 @@ public:
 private:
 	static FName ModuleName;
 };
+#endif //ANIMNEXT_TRACE_ENABLED
 
diff --git a/Engine/Plugins/Experimental/AnimNext/Source/AnimNextEditor/Private/RewindDebugger/RewindDebuggerAnimNext.cpp b/Engine/Plugins/Experimental/AnimNext/Source/AnimNextEditor/Private/RewindDebugger/RewindDebuggerAnimNext.cpp
--- a/Engine/Plugins/Experimental/AnimNext/Source/AnimNextEditor/Private/RewindDebugger/RewindDebuggerAnimNext.cpp
+++ b/Engine/Plugins/Experimental/AnimNext/Source/AnimNextEditor/Private/RewindDebugger/RewindDebuggerAnimNext.cpp
@@ -2,6 +2,7 @@
 
 #include "RewindDebuggerAnimNext.h"
 
+#if ANIMNEXT_TRACE_ENABLED
 #include "IGameplayProvider.h"
 
 FRewindDebuggerAnimNext::FRewindDebuggerAnimNext()
@@ -12,4 +13,6 @@ FRewindDebuggerAnimNext::FRewindDebuggerAnimNext()
 void FRewindDebuggerAnimNext::Update(float DeltaTime, IRewindDebugger* RewindDebugger)
 {
 	// Future place to sync anim next workspace editor debugging when scrubbing
-}
\ No newline at end of file
+}
+
+#endif // ANIMNEXT_TRACE_ENABLED
diff --git a/Engine/Plugins/Experimental/AnimNext/Source/AnimNextEditor/Private/RewindDebugger/RewindDebuggerAnimNext.h b/Engine/Plugins/Experimental/AnimNext/Source/AnimNextEditor/Private/RewindDebugger/RewindDebuggerAnimNext.h
--- a/Engine/Plugins/Experimental/AnimNext/Source/AnimNextEditor/Private/RewindDebugger/RewindDebuggerAnimNext.h
+++ b/Engine/Plugins/Experimental/AnimNext/Source/AnimNextEditor/Private/RewindDebugger/RewindDebuggerAnimNext.h
@@ -1,6 +1,9 @@
 // Copyright Epic Games, Inc. All Rights Reserved.
 #pragma once
 
+#include "RewindDebugger/AnimNextTrace.h"
+
+#if ANIMNEXT_TRACE_ENABLED
 #include "CoreMinimal.h"
 #include "IRewindDebugger.h"
 #include "IRewindDebuggerExtension.h"
@@ -15,3 +18,4 @@ public:
 
 	virtual void Update(float DeltaTime, IRewindDebugger* RewindDebugger) override;
 };
+#endif //ANIMNEXT_TRACE_ENABLED
