Fix inherited property debug

diff --git a/Engine/Plugins/Experimental/AnimNext/Source/AnimNext/Private/RewindDebugger/AnimNextTrace.cpp b/Engine/Plugins/Experimental/AnimNext/Source/AnimNext/Private/RewindDebugger/AnimNextTrace.cpp
--- a/Engine/Plugins/Experimental/AnimNext/Source/AnimNext/Private/RewindDebugger/AnimNextTrace.cpp
+++ b/Engine/Plugins/Experimental/AnimNext/Source/AnimNext/Private/RewindDebugger/AnimNextTrace.cpp
@@ -155,6 +155,12 @@ void FAnimNextTrace::OutputAnimNextVariables(const FAnimNextDataInterfaceInstanc
 	if (Variables.GetPropertyBagStruct())
 	{
 		TConstArrayView<FPropertyBagPropertyDesc> PropertyDescs = Variables.GetPropertyBagStruct()->GetPropertyDescs();
+		for(int i =0; i<PropertyDescs.Num(); i++)
+		{
+			uint8* ValueMemory = Instance->GetMemoryForVariable(i, PropertyDescs[i].Name, PropertyDescs[i].CachedProperty);
+			PropertyDescs[i].CachedProperty->CopySingleValue(const_cast<uint8*>(Variables.GetValue().GetMemory()) +  PropertyDescs[i].CachedProperty->GetOffset_ForInternal(), ValueMemory);
+		}
+		
 		uint32 PropertyDescHash = GetPropertyDescHash(PropertyDescs);
 		{
 			FRWScopeLock RWScopeLock(GTracedPropertiesRWLock, SLT_ReadOnly);
