Add support for material and morph target curves in AnimNext

diff --git a/Engine/Plugins/Experimental/AnimNext/Source/AnimNext/Internal/Component/SkinnedMeshComponentExtensions.h b/Engine/Plugins/Experimental/AnimNext/Source/AnimNext/Internal/Component/SkinnedMeshComponentExtensions.h
--- a/Engine/Plugins/Experimental/AnimNext/Source/AnimNext/Internal/Component/SkinnedMeshComponentExtensions.h
+++ b/Engine/Plugins/Experimental/AnimNext/Source/AnimNext/Internal/Component/SkinnedMeshComponentExtensions.h
@@ -19,7 +19,8 @@ struct FSkinnedMeshComponentExtensions
 		USkinnedMeshComponent* InComponent,
 		TConstArrayView<FBoneIndexType> InParentIndices,
 		TConstArrayView<FBoneIndexType> InRequiredBoneIndices,
-		TConstArrayView<FTransform> InLocalSpaceTransforms);
+		TConstArrayView<FTransform> InLocalSpaceTransforms,
+		TUniqueFunction<void(USkinnedMeshComponent*)>&& InGameThreadCallback);
 };
 
 }
diff --git a/Engine/Plugins/Experimental/AnimNext/Source/AnimNext/Private/Component/SkinnedMeshComponentExtensions.cpp b/Engine/Plugins/Experimental/AnimNext/Source/AnimNext/Private/Component/SkinnedMeshComponentExtensions.cpp
--- a/Engine/Plugins/Experimental/AnimNext/Source/AnimNext/Private/Component/SkinnedMeshComponentExtensions.cpp
+++ b/Engine/Plugins/Experimental/AnimNext/Source/AnimNext/Private/Component/SkinnedMeshComponentExtensions.cpp
@@ -13,8 +13,11 @@ void FSkinnedMeshComponentExtensions::CompleteAndDispatch(
 	USkinnedMeshComponent* InComponent,
 	TConstArrayView<FBoneIndexType> InParentIndices,
 	TConstArrayView<FBoneIndexType> InRequiredBoneIndices,
-	TConstArrayView<FTransform> InLocalSpaceTransforms)
+	TConstArrayView<FTransform> InLocalSpaceTransforms,
+	TUniqueFunction<void(USkinnedMeshComponent*)>&& InGameThreadCallback)
 {
+	bool bNeedsChildTransformUpdate = false;
+
 	// Fill the component space transform buffer
 	TArrayView<FTransform> ComponentSpaceTransforms = InComponent->GetEditableComponentSpaceTransforms();
 	if (ComponentSpaceTransforms.Num() > 0)
@@ -34,25 +37,40 @@ void FSkinnedMeshComponentExtensions::CompleteAndDispatch(
 		InComponent->SendRenderDynamicData_Concurrent();
 		InComponent->SendRenderTransform_Concurrent();
 
-		if (InComponent->bHasSocketAttachments)
+		bNeedsChildTransformUpdate = InComponent->bHasSocketAttachments;
+
+		if (const USkeletalMeshComponent* SkeletalMeshComponent = Cast<USkeletalMeshComponent>(InComponent))
 		{
-			FAnimNextModuleInstance::RunTaskOnGameThread([WeakComponent = TWeakObjectPtr<USkinnedMeshComponent>(InComponent)]()
+			TRACE_SKELETAL_MESH_COMPONENT(SkeletalMeshComponent);
+		}
+	}
+
+	if (bNeedsChildTransformUpdate || InGameThreadCallback)
+	{
+		FAnimNextModuleInstance::RunTaskOnGameThread(
+			[
+				WeakComponent = TWeakObjectPtr<USkinnedMeshComponent>(InComponent),
+				bNeedsChildTransformUpdate,
+				GameThreadCallback = MoveTemp(InGameThreadCallback)
+			]()
 			{
-				SCOPED_NAMED_EVENT(AnimNext_SkinnedMesh_CompleteAndDispatch_GameThread, FColor::Orange);
 				USkinnedMeshComponent* Component = WeakComponent.Get();
 				if (Component == nullptr)
 				{
 					return;
 				}
 
-				Component->UpdateChildTransforms(EUpdateTransformFlags::OnlyUpdateIfUsingSocket);
+				if (bNeedsChildTransformUpdate)
+				{
+					SCOPED_NAMED_EVENT(AnimNext_SkinnedMesh_CompleteAndDispatch_GameThread, FColor::Orange);
+					Component->UpdateChildTransforms(EUpdateTransformFlags::OnlyUpdateIfUsingSocket);
+				}
+
+				if (GameThreadCallback)
+				{
+					GameThreadCallback(Component);
+				}
 			});
-		}
-		
-		if (const USkeletalMeshComponent* SkeletalMeshComponent = Cast<USkeletalMeshComponent>(InComponent))
-		{
-			TRACE_SKELETAL_MESH_COMPONENT(SkeletalMeshComponent);
-		}
 	}
 }
 
diff --git a/Engine/Plugins/Experimental/AnimNext/Source/AnimNext/Private/GenerationTools.cpp b/Engine/Plugins/Experimental/AnimNext/Source/AnimNext/Private/GenerationTools.cpp
--- a/Engine/Plugins/Experimental/AnimNext/Source/AnimNext/Private/GenerationTools.cpp
+++ b/Engine/Plugins/Experimental/AnimNext/Source/AnimNext/Private/GenerationTools.cpp
@@ -42,7 +42,7 @@ struct FCompareBoneIndexType
 
 
 /*static*/ bool FGenerationTools::GenerateReferencePose(const USkeletalMeshComponent* SkeletalMeshComponent
-	, const USkeletalMesh* SkeletalMesh
+	, USkeletalMesh* SkeletalMesh
 	, FReferencePose& OutAnimationReferencePose)
 {
 	SCOPE_CYCLE_COUNTER(STAT_AnimNext_GenerateReferencePose);
@@ -194,6 +194,57 @@ struct FCompareBoneIndexType
 					LODBoneIndexToParentLODBoneIndexMap[LODBoneIndex] = ParentLODBoneIndex;
 				}
 
+				UE::Anim::FBulkCurveFlags CurveFlags;
+				Skeleton->ForEachCurveMetaData([&CurveFlags](const FName& InCurveName, const FCurveMetaData& InMetaData)
+					{
+						UE::Anim::ECurveElementFlags Flags = UE::Anim::ECurveElementFlags::None;
+
+						if (InMetaData.Type.bMaterial)
+						{
+							Flags |= UE::Anim::ECurveElementFlags::Material;
+						}
+
+						if (InMetaData.Type.bMorphtarget)
+						{
+							Flags |= UE::Anim::ECurveElementFlags::MorphTarget;
+						}
+
+						if (Flags != UE::Anim::ECurveElementFlags::None)
+						{
+							// Add curve with any relevant flags to bulk flags
+							CurveFlags.Add(InCurveName, Flags);
+						}
+					});
+
+				// Override any metadata with the skeletal mesh
+				if (const UAnimCurveMetaData* MetaDataUserData = SkeletalMesh->GetAssetUserData<UAnimCurveMetaData>())
+				{
+					UE::Anim::FBulkCurveFlags MeshCurveFlags;
+
+					// Apply morph target flags to any morph curves
+					MetaDataUserData->ForEachCurveMetaData([&MeshCurveFlags](FName InCurveName, const FCurveMetaData& InCurveMetaData)
+						{
+							UE::Anim::ECurveElementFlags Flags = UE::Anim::ECurveElementFlags::None;
+
+							if (InCurveMetaData.Type.bMaterial)
+							{
+								Flags |= UE::Anim::ECurveElementFlags::Material;
+							}
+
+							if (InCurveMetaData.Type.bMorphtarget)
+							{
+								Flags |= UE::Anim::ECurveElementFlags::MorphTarget;
+							}
+
+							if (Flags != UE::Anim::ECurveElementFlags::None)
+							{
+								MeshCurveFlags.Add(InCurveName, Flags);
+							}
+						});
+
+					UE::Anim::FNamedValueArrayUtils::Union(CurveFlags, MeshCurveFlags);
+				}
+
 				OutAnimationReferencePose.Initialize(
 					RefSkeleton,
 					{ LODBoneIndexToParentLODBoneIndexMap },
@@ -203,6 +254,7 @@ struct FCompareBoneIndexType
 					{ MeshBoneIndexToLODBoneIndexMap },
 					LODNumBones,
 					BoneNameToLODIndexMap,
+					CurveFlags,
 					bCanGenerateSingleBonesList);
 
 				ReferencePoseGenerated = true;
@@ -879,7 +931,7 @@ namespace AnimNext::Tools::ConsoleCommands
 			const int32 NumAssets = Assets.Num();
 			for (int32 Idx = 0; Idx < NumAssets; Idx++)
 			{
-				const USkeletalMesh* SkeletalMesh = Cast<USkeletalMesh>(Assets[Idx].GetAsset());
+				USkeletalMesh* SkeletalMesh = Cast<USkeletalMesh>(Assets[Idx].GetAsset());
 				if (SkeletalMesh != nullptr)
 				{
 					UE::AnimNext::FReferencePose OutAnimationReferencePose;
diff --git a/Engine/Plugins/Experimental/AnimNext/Source/AnimNext/Private/Graph/RigUnit_AnimNextWriteSkeletalMeshComponentPose.cpp b/Engine/Plugins/Experimental/AnimNext/Source/AnimNext/Private/Graph/RigUnit_AnimNextWriteSkeletalMeshComponentPose.cpp
--- a/Engine/Plugins/Experimental/AnimNext/Source/AnimNext/Private/Graph/RigUnit_AnimNextWriteSkeletalMeshComponentPose.cpp
+++ b/Engine/Plugins/Experimental/AnimNext/Source/AnimNext/Private/Graph/RigUnit_AnimNextWriteSkeletalMeshComponentPose.cpp
@@ -2,6 +2,8 @@
 
 #include "RigUnit_AnimNextWriteSkeletalMeshComponentPose.h"
 
+#include "AnimationRuntime.h"
+#include "Animation/NamedValueArray.h"
 #include "AnimNextStats.h"
 #include "GenerationTools.h"
 #include "Component/SkinnedMeshComponentExtensions.h"
@@ -76,6 +78,8 @@ FRigUnit_AnimNextWriteSkeletalMeshComponentPose_Execute()
 	const int32 SkeletalMeshComponentLODLevel = OutputComponent->GetPredictedLODLevel();
 	const int32 PoseLODLevel = Pose.LODPose.LODLevel;
 
+	const TArrayView<const FBoneIndexType> PoseBoneToMeshBoneIndexMap = RefPose.GetLODBoneIndexToMeshBoneIndexMap(SkeletalMeshComponentLODLevel);
+
 	TArray<FTransform, TMemStackAllocator<>> LocalSpaceTransforms;
 
 	const FReferenceSkeleton& RefSkeleton = SkeletalMesh->GetRefSkeleton();
@@ -89,25 +93,28 @@ FRigUnit_AnimNextWriteSkeletalMeshComponentPose_Execute()
 	if (SkeletalMeshComponentLODLevel < PoseLODLevel)
 	{
 		// The given pose is missing transforms that the visual representation needs, initialize these with the reference pose.
-
-		const TArrayView<const FBoneIndexType> PoseBoneToMeshBoneIndexMap = RefPose.GetLODBoneIndexToMeshBoneIndexMap(SkeletalMeshComponentLODLevel);
 		const int32 NumPoseBones = Pose.LODPose.LocalTransforms.Num();
 		const int32 NumSkeletalMeshBones = PoseBoneToMeshBoneIndexMap.Num();
 
+		const TArray<FTransform>& MeshBindPose = RefSkeleton.GetRefBonePose();
+
 		for (int32 PoseBoneIndex = NumPoseBones; PoseBoneIndex < NumSkeletalMeshBones; ++PoseBoneIndex)
 		{
 			const int32 MeshBoneIndex = PoseBoneToMeshBoneIndexMap[PoseBoneIndex];
-			LocalSpaceTransforms[MeshBoneIndex] = RefSkeleton.GetRefBonePose()[MeshBoneIndex];
+
+			// We write to the mesh bone index (remapped index) directly since it won't be written below
+			LocalSpaceTransforms[MeshBoneIndex] = MeshBindPose[MeshBoneIndex];
 		}
 	}
 
+	FBlendedHeapCurve PrevCurves;
+
 	// Clear our curves and attributes or we'll have leftovers from our last write.
 	PRAGMA_DISABLE_DEPRECATION_WARNINGS
-	OutputComponent->AnimCurves.Empty();
+	Swap(PrevCurves, OutputComponent->AnimCurves);
 	PRAGMA_ENABLE_DEPRECATION_WARNINGS
 	OutputComponent->GetEditableCustomAttributes().Empty();
 
-
 	// Map LOD pose into local-space scratch buffer
 	FGenerationTools::RemapPose(Pose.LODPose, LocalSpaceTransforms);
 
@@ -115,6 +122,111 @@ FRigUnit_AnimNextWriteSkeletalMeshComponentPose_Execute()
 	OutputComponent->AnimCurves.CopyFrom(Pose.Curves);
 	PRAGMA_ENABLE_DEPRECATION_WARNINGS
 
+#if WITH_EDITOR
+	PRAGMA_DISABLE_DEPRECATION_WARNINGS
+	OutputComponent->GetEditableAnimationCurves().CopyFrom(Pose.Curves);
+	PRAGMA_ENABLE_DEPRECATION_WARNINGS
+#endif
+
+	TUniqueFunction<void(USkinnedMeshComponent*)> GameThreadCallback;
+
+	// Extract our material parameters and morph target curves
+	if (Pose.Curves.Num() != 0 && RefPose.CurveFlags.Num() != 0)
+	{
+		FBlendedCurve ExpiredCurves;
+		ExpiredCurves.Reserve(64);			// Expired curves are uncommon
+
+		// Reset any material curves we produced in the last update but not in this update
+		UE::Anim::FNamedValueArrayUtils::Subtraction(PrevCurves, Pose.Curves,
+			[&ExpiredCurves](const UE::Anim::FCurveElement& InCurveElement)
+			{
+				ExpiredCurves.Add(InCurveElement);
+			});
+
+		FBlendedHeapCurve ExpiredMaterialCurves;
+		ExpiredMaterialCurves.Reserve(ExpiredCurves.Num());
+
+		FBlendedHeapCurve MaterialCurves;
+		MaterialCurves.Reserve(RefPose.CurveFlags.Num());
+
+		TMap<FName, float> MorphTargetCurves;
+		MorphTargetCurves.Reserve(RefPose.CurveFlags.Num());
+
+		UE::Anim::FNamedValueArrayUtils::Intersection(ExpiredCurves, RefPose.CurveFlags,
+			[&ExpiredMaterialCurves](const UE::Anim::FCurveElement& InCurveElement, const UE::Anim::FCurveElementFlags& InCurveFlagsElement)
+			{
+				if (EnumHasAnyFlags(InCurveFlagsElement.Flags | InCurveElement.Flags, UE::Anim::ECurveElementFlags::Material))
+				{
+					ExpiredMaterialCurves.Add(InCurveElement);
+				}
+			});
+
+		UE::Anim::FNamedValueArrayUtils::Intersection(Pose.Curves, RefPose.CurveFlags,
+			[&MaterialCurves, &MorphTargetCurves](const UE::Anim::FCurveElement& InCurveElement, const UE::Anim::FCurveElementFlags& InCurveFlagsElement)
+			{
+				if (EnumHasAnyFlags(InCurveFlagsElement.Flags | InCurveElement.Flags, UE::Anim::ECurveElementFlags::MorphTarget))
+				{
+					MorphTargetCurves.Add(InCurveElement.Name, InCurveElement.Value);
+				}
+
+				if (EnumHasAnyFlags(InCurveFlagsElement.Flags | InCurveElement.Flags, UE::Anim::ECurveElementFlags::Material))
+				{
+					MaterialCurves.Add(InCurveElement);
+				}
+			});
+
+		if (ExpiredMaterialCurves.Num() != 0 || MaterialCurves.Num() != 0 || !MorphTargetCurves.IsEmpty())
+		{
+			GameThreadCallback = [
+				ExpiredMaterialCurves = MoveTemp(ExpiredMaterialCurves),
+				MaterialCurves = MoveTemp(MaterialCurves),
+				MorphTargetCurves = MoveTemp(MorphTargetCurves)
+				](USkinnedMeshComponent* InComponent)
+				{
+					USkeletalMeshComponent* SMComponent = CastChecked<USkeletalMeshComponent>(InComponent);
+
+					SMComponent->ResetMorphTargetCurves();
+
+					if (!MorphTargetCurves.IsEmpty())
+					{
+						// TODO: We should avoid using a TMap for this, a named value array would be a better fit
+						USkeletalMesh* SkeletalMesh = SMComponent->GetSkeletalMeshAsset();
+						FAnimationRuntime::AppendActiveMorphTargets(SkeletalMesh, MorphTargetCurves, SMComponent->ActiveMorphTargets, SMComponent->MorphTargetWeights);
+					}
+
+					// now update morph target curves that are added via SetMorphTarget
+					SMComponent->UpdateMorphTargetOverrideCurves();
+
+					// Reset material curves that are no longer animated to their default value
+					ExpiredMaterialCurves.ForEachElement([SMComponent](const UE::Anim::FCurveElement& MaterialCurve)
+						{
+							float DefaultValue = SMComponent->GetScalarParameterDefaultValue(MaterialCurve.Name);
+							SMComponent->SetScalarParameterValueOnMaterials(MaterialCurve.Name, DefaultValue);
+						});
+
+					// Update material curves that were animated
+					MaterialCurves.ForEachElement([SMComponent](const UE::Anim::FCurveElement& MaterialCurve)
+						{
+							SMComponent->SetScalarParameterValueOnMaterials(MaterialCurve.Name, MaterialCurve.Value);
+						});
+				};
+		}
+	}
+
+	// If we don't have any animated curves, make sure we reset the morph targets and update the ones from gameplay
+	if (!GameThreadCallback)
+	{
+		GameThreadCallback = [](USkinnedMeshComponent* InComponent)
+			{
+				USkeletalMeshComponent* SMComponent = CastChecked<USkeletalMeshComponent>(InComponent);
+
+				SMComponent->ResetMorphTargetCurves();
+
+				// now update morph target curves that are added via SetMorphTarget
+				SMComponent->UpdateMorphTargetOverrideCurves();
+			};
+	}
+
 	// Attributes require remapping since the indices are LOD indices and we want mesh indices
 	FGenerationTools::RemapAttributes(Pose.LODPose, Pose.Attributes, OutputComponent->GetEditableCustomAttributes());
 
@@ -122,6 +234,7 @@ FRigUnit_AnimNextWriteSkeletalMeshComponentPose_Execute()
 	UE::Anim::FSkinnedMeshComponentExtensions::CompleteAndDispatch(
 		OutputComponent,
 		RefPose.GetMeshBoneIndexToParentMeshBoneIndexMap(),
-		RefPose.GetLODBoneIndexToMeshBoneIndexMap(SkeletalMeshComponentLODLevel), // Based on the LOD level of the skeletal mesh component
-		LocalSpaceTransforms);
+		PoseBoneToMeshBoneIndexMap, // Based on the LOD level of the skeletal mesh component
+		LocalSpaceTransforms,
+		MoveTemp(GameThreadCallback));
 }
diff --git a/Engine/Plugins/Experimental/AnimNext/Source/AnimNext/Public/GenerationTools.h b/Engine/Plugins/Experimental/AnimNext/Source/AnimNext/Public/GenerationTools.h
--- a/Engine/Plugins/Experimental/AnimNext/Source/AnimNext/Public/GenerationTools.h
+++ b/Engine/Plugins/Experimental/AnimNext/Source/AnimNext/Public/GenerationTools.h
@@ -34,7 +34,7 @@ public:
 	// If no SkeletalMeshComponent is passed, the reference pose will not exclude invisible bones and will not include shadow shapes required bones
 	// If no SkeletalMesh asset is passed, there will be no generation
 	static UE_API bool GenerateReferencePose(const USkeletalMeshComponent* SkeletalMeshComponent
-		, const USkeletalMesh* SkeletalMesh
+		, USkeletalMesh* SkeletalMesh
 		, UE::AnimNext::FReferencePose& OutAnimationReferencePose);
 
 	// Generates the full list of bones required by a LOD, based on the SkeletalRetrieve RequiredBones
diff --git a/Engine/Plugins/Experimental/AnimNext/Source/AnimNext/Public/ReferencePose.h b/Engine/Plugins/Experimental/AnimNext/Source/AnimNext/Public/ReferencePose.h
--- a/Engine/Plugins/Experimental/AnimNext/Source/AnimNext/Public/ReferencePose.h
+++ b/Engine/Plugins/Experimental/AnimNext/Source/AnimNext/Public/ReferencePose.h
@@ -71,6 +71,9 @@ struct TReferencePose
 	// Mapping of bone names to LOD Bone Indices
 	TMap<FName, FBoneIndexType, SetAllocator> BoneNameToLODBoneIndexMap;
 
+	// Curves of interest (material/morph target)
+	UE::Anim::FBulkCurveFlags CurveFlags;
+
 	TWeakObjectPtr<const USkeletalMeshComponent> SkeletalMeshComponent = nullptr;
 	TWeakObjectPtr<const USkeletalMesh> SkeletalMesh = nullptr;
 	TWeakObjectPtr<const USkeleton> Skeleton = nullptr;
@@ -110,6 +113,7 @@ struct TReferencePose
 		, const TArray<FBoneIndexType>& InMeshBoneIndexToLODBoneIndexMap
 		, const TArray<int32, AllocatorType>& InLODNumBones
 		, const TMap<FName, FBoneIndexType>& InNameToLODBoneIndexMap
+		, const UE::Anim::FBulkCurveFlags& InCurveFlags
 		, bool bFastPath)
 	{
 		const int32 NumBonesLOD0 = !InLODNumBones.IsEmpty() ? InLODNumBones[0] : 0;
@@ -123,6 +127,7 @@ struct TReferencePose
 		SkeletonBoneIndexToLODBoneIndexMap = InSkeletonBoneIndexToLODBoneIndexMap;
 		MeshBoneIndexToLODBoneIndexMap = InMeshBoneIndexToLODBoneIndexMap;
 		BoneNameToLODBoneIndexMap = InNameToLODBoneIndexMap;
+		CurveFlags = InCurveFlags;
 		LODNumBones = InLODNumBones;
 
 		const TArray<FTransform>& RefBonePose = RefSkeleton.GetRefBonePose();
diff --git a/Engine/Source/Runtime/Engine/Public/Animation/NamedValueArray.h b/Engine/Source/Runtime/Engine/Public/Animation/NamedValueArray.h
--- a/Engine/Source/Runtime/Engine/Public/Animation/NamedValueArray.h
+++ b/Engine/Source/Runtime/Engine/Public/Animation/NamedValueArray.h
@@ -698,6 +698,77 @@ struct FNamedValueArrayUtils
 		}
 	}
 
+	/**
+	 * Calls predicate on all elements of the first named value array that are not present in the second named value array: Result = (A - B).
+	 * ValuePredicateType is a function of signature: (const ElementType0& InElement0) -> void
+	 **/
+	template<typename AllocatorType0, typename ElementType0, typename AllocatorType1, typename ElementType1, typename ValuePredicateType>
+	static void Subtraction(const TNamedValueArray<AllocatorType0, ElementType0>& InNamedValues0, const TNamedValueArray<AllocatorType1, ElementType1>& InNamedValues1, ValuePredicateType InValuePredicate)
+	{
+		CURVE_PROFILE_CYCLE_COUNTER(FNamedValueArrayUtils_Subtraction);
+
+		// Check arrays are not overlapping
+		checkSlow((void*)&InNamedValues0 != (void*)&InNamedValues1);
+
+		// Sort both inputs if required
+		InNamedValues0.SortElementsIfRequired();
+		InNamedValues1.SortElementsIfRequired();
+
+		const int32 NumElements0 = InNamedValues0.Num();
+		const int32 NumElements1 = InNamedValues1.Num();
+
+		// Perform dual-iteration on the two sorted arrays
+		int32 ElementIndex0 = 0;
+		int32 ElementIndex1 = 0;
+
+		while (true)
+		{
+			if (ElementIndex0 == NumElements0 && ElementIndex1 < NumElements1)
+			{
+				// Reached end of the first array, we can just early out
+				break;
+			}
+			else if (ElementIndex1 == NumElements1 && ElementIndex0 < NumElements0)
+			{
+				// Reached end of the second array, we can just early out
+				break;
+			}
+			else if (ElementIndex0 == NumElements0 && ElementIndex1 == NumElements1)
+			{
+				// All elements exhausted, exit
+				break;
+			}
+
+			const ElementType0* RESTRICT Element0 = &InNamedValues0.Elements[ElementIndex0];
+			const ElementType1* RESTRICT Element1 = &InNamedValues1.Elements[ElementIndex1];
+
+			if (Element0->Name == Element1->Name)
+			{
+				// Element exists in both arrays
+				++ElementIndex0;
+				++ElementIndex1;
+			}
+			else if (Element0->Name.FastLess(Element1->Name))
+			{
+				// Element exists only in first array
+				InValuePredicate(*Element0);
+				++ElementIndex0;
+			}
+			else
+			{
+				// Element exists only in second array
+				++ElementIndex1;
+			}
+		}
+
+		// Drain any remaining elements from our first array
+		while (ElementIndex0 < NumElements0)
+		{
+			InValuePredicate(InNamedValues0.Elements[ElementIndex0]);
+			++ElementIndex0;
+		}
+	}
+
 	/**
 	 * Removes elements in InOutValueArray0 that match InValueArray1 if predicate returns false
 	 **/
