Verify user can Undo/Redo in AnimNext Module
Automated QMetry Test Case: UE-TC-18778

diff --git a/Engine/Plugins/Experimental/AnimNext/Source/AnimNextEditor/Internal/Common/GraphEditorSchemaActions.h b/Engine/Plugins/Experimental/AnimNext/Source/AnimNextEditor/Internal/Common/GraphEditorSchemaActions.h
--- a/Engine/Plugins/Experimental/AnimNext/Source/AnimNextEditor/Internal/Common/GraphEditorSchemaActions.h
+++ b/Engine/Plugins/Experimental/AnimNext/Source/AnimNextEditor/Internal/Common/GraphEditorSchemaActions.h
@@ -80,7 +80,7 @@ private:
 
 
 USTRUCT()
-struct FAnimNextSchemaAction_Variable : public FAnimNextSchemaAction
+struct ANIMNEXTEDITOR_API FAnimNextSchemaAction_Variable : public FAnimNextSchemaAction
 {
 	GENERATED_BODY()
 
@@ -153,4 +153,4 @@ private:
 	/** Holds the node type that this spawner will instantiate. */
 	UPROPERTY(Transient)
 	TSubclassOf<UEdGraphNode> NodeClass;
-};
\ No newline at end of file
+};
diff --git a/Engine/Plugins/Experimental/AnimNext/Source/AnimNextUncookedOnly/Internal/AnimNextRigVMAssetEditorData.h b/Engine/Plugins/Experimental/AnimNext/Source/AnimNextUncookedOnly/Internal/AnimNextRigVMAssetEditorData.h
--- a/Engine/Plugins/Experimental/AnimNext/Source/AnimNextUncookedOnly/Internal/AnimNextRigVMAssetEditorData.h
+++ b/Engine/Plugins/Experimental/AnimNext/Source/AnimNextUncookedOnly/Internal/AnimNextRigVMAssetEditorData.h
@@ -205,6 +205,7 @@ protected:
 	friend UE::AnimNext::UncookedOnly::FScopedCompilerResults;
 	friend UE::AnimNext::Editor::FAssetCompilationHandler;
 	friend class UAnimNextStateTreeTreeEditorData;
+	friend struct FAnimNextModuleTests;
 
 	// UObject interface
 	virtual void Serialize(FArchive& Ar) override;
