Defer render data update to the end of frame

diff --git a/Engine/Plugins/Experimental/AnimNext/Source/AnimNext/Private/Component/SkinnedMeshComponentExtensions.cpp b/Engine/Plugins/Experimental/AnimNext/Source/AnimNext/Private/Component/SkinnedMeshComponentExtensions.cpp
--- a/Engine/Plugins/Experimental/AnimNext/Source/AnimNext/Private/Component/SkinnedMeshComponentExtensions.cpp
+++ b/Engine/Plugins/Experimental/AnimNext/Source/AnimNext/Private/Component/SkinnedMeshComponentExtensions.cpp
@@ -6,6 +6,15 @@
 #include "Animation/AnimTrace.h"
 #include "Module/AnimNextModuleInstance.h"
 
+namespace UE::Anim::Private
+{
+	// When enabled, we mark the component render data as being dirty on the main thread and we'll send the data
+	// to the render thread at the end of the frame, same as AnimBP
+	// When disabled, we do it inline here
+	static bool GUseDeferredRenderDataUpdate = true;
+	FAutoConsoleVariableRef CVar_UseDeferredRenderDataUpdate(TEXT("a.AnimNext.UseDeferredRenderDataUpdate"), GUseDeferredRenderDataUpdate, TEXT("Whether or not to defer the render data update to the end of frame"));
+}
+
 namespace UE::Anim
 {
 
@@ -33,9 +42,12 @@ void FSkinnedMeshComponentExtensions::CompleteAndDispatch(
 		InComponent->InvalidateCachedBounds();
 		InComponent->UpdateBounds();
 
-		// Send updated transforms & bounds to the renderer
-		InComponent->SendRenderDynamicData_Concurrent();
-		InComponent->SendRenderTransform_Concurrent();
+		if (!UE::Anim::Private::GUseDeferredRenderDataUpdate)
+		{
+			// Send updated transforms & bounds to the renderer
+			InComponent->SendRenderDynamicData_Concurrent();
+			InComponent->SendRenderTransform_Concurrent();
+		}
 
 		bNeedsChildTransformUpdate = InComponent->bHasSocketAttachments;
 
@@ -45,7 +57,7 @@ void FSkinnedMeshComponentExtensions::CompleteAndDispatch(
 		}
 	}
 
-	if (bNeedsChildTransformUpdate || InGameThreadCallback)
+	if (bNeedsChildTransformUpdate || InGameThreadCallback || UE::Anim::Private::GUseDeferredRenderDataUpdate)
 	{
 		FAnimNextModuleInstance::RunTaskOnGameThread(
 			[
@@ -70,6 +82,15 @@ void FSkinnedMeshComponentExtensions::CompleteAndDispatch(
 				{
 					GameThreadCallback(Component);
 				}
+
+				if (UE::Anim::Private::GUseDeferredRenderDataUpdate)
+				{
+					// Need to send new bounds to the render thread
+					Component->MarkRenderTransformDirty();
+
+					// New bone positions need to be sent to render thread
+					Component->MarkRenderDynamicDataDirty();
+				}
 			});
 	}
 }
