Remove write pose intermediate buffer

diff --git a/Engine/Plugins/Experimental/AnimNext/Source/AnimNext/Internal/Component/SkinnedMeshComponentExtensions.h b/Engine/Plugins/Experimental/AnimNext/Source/AnimNext/Internal/Component/SkinnedMeshComponentExtensions.h
--- a/Engine/Plugins/Experimental/AnimNext/Source/AnimNext/Internal/Component/SkinnedMeshComponentExtensions.h
+++ b/Engine/Plugins/Experimental/AnimNext/Source/AnimNext/Internal/Component/SkinnedMeshComponentExtensions.h
@@ -4,7 +4,7 @@
 
 #include "Containers/ContainersFwd.h"
 #include "Math/MathFwd.h"
-#include "BoneIndices.h"
+#include "LODPose.h"
 
 class USkinnedMeshComponent;
 
@@ -17,9 +17,7 @@ struct FSkinnedMeshComponentExtensions
 	// Converts local to component space, flips buffers, updates bounds and dispatches to renderer
 	static ANIMNEXT_API void CompleteAndDispatch(
 		USkinnedMeshComponent* InComponent,
-		TConstArrayView<FBoneIndexType> InParentIndices,
-		TConstArrayView<FBoneIndexType> InRequiredBoneIndices,
-		TConstArrayView<FTransform> InLocalSpaceTransforms,
+		const UE::AnimNext::FLODPoseHeap& InLODPose,
 		TUniqueFunction<void(USkinnedMeshComponent*)>&& InGameThreadCallback);
 };
 
diff --git a/Engine/Plugins/Experimental/AnimNext/Source/AnimNext/Private/Component/SkinnedMeshComponentExtensions.cpp b/Engine/Plugins/Experimental/AnimNext/Source/AnimNext/Private/Component/SkinnedMeshComponentExtensions.cpp
--- a/Engine/Plugins/Experimental/AnimNext/Source/AnimNext/Private/Component/SkinnedMeshComponentExtensions.cpp
+++ b/Engine/Plugins/Experimental/AnimNext/Source/AnimNext/Private/Component/SkinnedMeshComponentExtensions.cpp
@@ -18,11 +18,87 @@ namespace UE::Anim::Private
 namespace UE::Anim
 {
 
+void ConvertLocalSpaceToComponentSpace(
+	const USkinnedMeshComponent* InComponent,
+	const UE::AnimNext::FLODPoseHeap& InLODPose,
+	TArrayView<FTransform> OutComponentSpaceTransforms)
+{
+	SCOPE_CYCLE_COUNTER(STAT_AnimNext_ConvertLocalSpaceToComponentSpace);
+
+	const UE::AnimNext::FReferencePose& RefPose = InLODPose.GetRefPose();
+
+	const TArrayView<const FBoneIndexType> MeshBoneIndexToParentMeshBoneIndexMap = RefPose.GetMeshBoneIndexToParentMeshBoneIndexMap();
+	const TArrayView<const FBoneIndexType> LODBoneIndexToMeshBoneIndexMap = RefPose.GetLODBoneIndexToMeshBoneIndexMap(InLODPose.LODLevel);
+
+	checkf(LODBoneIndexToMeshBoneIndexMap.Num() == InLODPose.GetNumBones(), TEXT("Buffer mismatch: %d:%d"), LODBoneIndexToMeshBoneIndexMap.Num(), InLODPose.GetNumBones());
+	checkf(MeshBoneIndexToParentMeshBoneIndexMap.Num() == OutComponentSpaceTransforms.Num(), TEXT("Buffer mismatch: %d:%d"), MeshBoneIndexToParentMeshBoneIndexMap.Num(), OutComponentSpaceTransforms.Num());
+
+	const int32 NumPoseLODBones = InLODPose.GetNumBones();
+	const FBoneIndexType* RESTRICT MeshBoneIndexToParentMeshBoneIndexMapPtr = MeshBoneIndexToParentMeshBoneIndexMap.GetData();
+	const FBoneIndexType* RESTRICT LODBoneIndexToMeshBoneIndexMapPtr = LODBoneIndexToMeshBoneIndexMap.GetData();
+	FTransform* RESTRICT ComponentSpaceData = OutComponentSpaceTransforms.GetData();
+
+	// First bone (if we have one) is always root bone, and it doesn't have a parent.
+	{
+		check(LODBoneIndexToMeshBoneIndexMap[0] == 0);
+		OutComponentSpaceTransforms[0] = InLODPose.LocalTransformsView[0];
+	}
+
+	for (int32 LODBoneIndex = 1; LODBoneIndex < NumPoseLODBones; ++LODBoneIndex)
+	{
+		const FTransform LocalSpaceTransform = InLODPose.LocalTransformsView[LODBoneIndex];
+
+		const FBoneIndexType MeshBoneIndex = LODBoneIndexToMeshBoneIndexMapPtr[LODBoneIndex];
+		const FBoneIndexType ParentMeshBoneIndex = MeshBoneIndexToParentMeshBoneIndexMapPtr[MeshBoneIndex];
+
+		const FTransform* RESTRICT ParentComponentSpaceTransform = ComponentSpaceData + ParentMeshBoneIndex;
+
+		FTransform* RESTRICT ComponentSpaceTransform = ComponentSpaceData + MeshBoneIndex;
+
+		FTransform::Multiply(ComponentSpaceTransform, &LocalSpaceTransform, ParentComponentSpaceTransform);
+
+		ComponentSpaceTransform->NormalizeRotation();
+
+		checkSlow(ComponentSpaceTransform->IsRotationNormalized());
+		checkSlow(!ComponentSpaceTransform->ContainsNaN());
+	}
+
+	const int32 MeshLOD = InComponent->GetPredictedLODLevel();
+	if (MeshLOD < InLODPose.LODLevel)
+	{
+		// If we evaluate the animation with lower quality than the visual mesh, we might
+		// have some transforms that we didn't compute. We'll grab the reference pose for them.
+		const TArrayView<const FBoneIndexType> LODBoneIndexToMeshBoneIndexMapForMeshLOD = RefPose.GetLODBoneIndexToMeshBoneIndexMap(MeshLOD);
+
+		const int32 NumMeshLODBones = LODBoneIndexToMeshBoneIndexMapForMeshLOD.Num();
+		check(NumPoseLODBones <= NumMeshLODBones);
+
+		const FBoneIndexType* RESTRICT LODBoneIndexToMeshBoneIndexMapForMeshLODPtr = LODBoneIndexToMeshBoneIndexMapForMeshLOD.GetData();
+
+		for (int32 LODBoneIndex = NumPoseLODBones; LODBoneIndex < NumMeshLODBones; ++LODBoneIndex)
+		{
+			const FTransform LocalSpaceTransform = RefPose.GetRefPoseTransform(LODBoneIndex);
+
+			const FBoneIndexType MeshBoneIndex = LODBoneIndexToMeshBoneIndexMapForMeshLODPtr[LODBoneIndex];
+			const FBoneIndexType ParentMeshBoneIndex = MeshBoneIndexToParentMeshBoneIndexMapPtr[MeshBoneIndex];
+
+			const FTransform* RESTRICT ParentComponentSpaceTransform = ComponentSpaceData + ParentMeshBoneIndex;
+
+			FTransform* RESTRICT ComponentSpaceTransform = ComponentSpaceData + MeshBoneIndex;
+
+			FTransform::Multiply(ComponentSpaceTransform, &LocalSpaceTransform, ParentComponentSpaceTransform);
+
+			ComponentSpaceTransform->NormalizeRotation();
+
+			checkSlow(ComponentSpaceTransform->IsRotationNormalized());
+			checkSlow(!ComponentSpaceTransform->ContainsNaN());
+		}
+	}
+}
+
 void FSkinnedMeshComponentExtensions::CompleteAndDispatch(
 	USkinnedMeshComponent* InComponent,
-	TConstArrayView<FBoneIndexType> InParentIndices,
-	TConstArrayView<FBoneIndexType> InRequiredBoneIndices,
-	TConstArrayView<FTransform> InLocalSpaceTransforms,
+	const UE::AnimNext::FLODPoseHeap& InLODPose,
 	TUniqueFunction<void(USkinnedMeshComponent*)>&& InGameThreadCallback)
 {
 	bool bNeedsChildTransformUpdate = false;
@@ -31,7 +107,7 @@ void FSkinnedMeshComponentExtensions::CompleteAndDispatch(
 	TArrayView<FTransform> ComponentSpaceTransforms = InComponent->GetEditableComponentSpaceTransforms();
 	if (ComponentSpaceTransforms.Num() > 0)
 	{
-		UE::AnimNext::FGenerationTools::ConvertLocalSpaceToComponentSpace(InParentIndices, InLocalSpaceTransforms, InRequiredBoneIndices, ComponentSpaceTransforms);
+		ConvertLocalSpaceToComponentSpace(InComponent, InLODPose, ComponentSpaceTransforms);
 
 		// Flag buffer for flip
 		InComponent->bNeedToFlipSpaceBaseBuffers = true;
diff --git a/Engine/Plugins/Experimental/AnimNext/Source/AnimNext/Private/Graph/RigUnit_AnimNextWriteSkeletalMeshComponentPose.cpp b/Engine/Plugins/Experimental/AnimNext/Source/AnimNext/Private/Graph/RigUnit_AnimNextWriteSkeletalMeshComponentPose.cpp
--- a/Engine/Plugins/Experimental/AnimNext/Source/AnimNext/Private/Graph/RigUnit_AnimNextWriteSkeletalMeshComponentPose.cpp
+++ b/Engine/Plugins/Experimental/AnimNext/Source/AnimNext/Private/Graph/RigUnit_AnimNextWriteSkeletalMeshComponentPose.cpp
@@ -74,39 +74,6 @@ FRigUnit_AnimNextWriteSkeletalMeshComponentPose_Execute()
 	FMemMark MemMark(FMemStack::Get());
 	USkinnedMeshComponent::FRenderStateLockScope LockScope(OutputComponent);
 
-	// TODO: The LOD level we got here is not guaranteed to be stable. We should find another way to get/bind the skeletal mesh component's LOD level.
-	const int32 SkeletalMeshComponentLODLevel = OutputComponent->GetPredictedLODLevel();
-	const int32 PoseLODLevel = Pose.LODPose.LODLevel;
-
-	const TArrayView<const FBoneIndexType> PoseBoneToMeshBoneIndexMap = RefPose.GetLODBoneIndexToMeshBoneIndexMap(SkeletalMeshComponentLODLevel);
-
-	TArray<FTransform, TMemStackAllocator<>> LocalSpaceTransforms;
-
-	const FReferenceSkeleton& RefSkeleton = SkeletalMesh->GetRefSkeleton();
-	LocalSpaceTransforms.SetNumUninitialized(RefSkeleton.GetNum());
-
-#if UE_ENABLE_POSE_DEBUG_FILL
-	FMemory::Memset(LocalSpaceTransforms.GetData(), 0xCD, sizeof(FTransform) * LocalSpaceTransforms.Num());
-#endif
-
-	// Did we evaluate animation with lower quality than the visual mesh?
-	if (SkeletalMeshComponentLODLevel < PoseLODLevel)
-	{
-		// The given pose is missing transforms that the visual representation needs, initialize these with the reference pose.
-		const int32 NumPoseBones = Pose.LODPose.LocalTransforms.Num();
-		const int32 NumSkeletalMeshBones = PoseBoneToMeshBoneIndexMap.Num();
-
-		const TArray<FTransform>& MeshBindPose = RefSkeleton.GetRefBonePose();
-
-		for (int32 PoseBoneIndex = NumPoseBones; PoseBoneIndex < NumSkeletalMeshBones; ++PoseBoneIndex)
-		{
-			const int32 MeshBoneIndex = PoseBoneToMeshBoneIndexMap[PoseBoneIndex];
-
-			// We write to the mesh bone index (remapped index) directly since it won't be written below
-			LocalSpaceTransforms[MeshBoneIndex] = MeshBindPose[MeshBoneIndex];
-		}
-	}
-
 	FBlendedHeapCurve PrevCurves;
 
 	// Clear our curves and attributes or we'll have leftovers from our last write.
@@ -115,9 +82,6 @@ FRigUnit_AnimNextWriteSkeletalMeshComponentPose_Execute()
 	PRAGMA_ENABLE_DEPRECATION_WARNINGS
 	OutputComponent->GetEditableCustomAttributes().Empty();
 
-	// Map LOD pose into local-space scratch buffer
-	FGenerationTools::RemapPose(Pose.LODPose, LocalSpaceTransforms);
-
 	PRAGMA_DISABLE_DEPRECATION_WARNINGS
 	OutputComponent->AnimCurves.CopyFrom(Pose.Curves);
 	PRAGMA_ENABLE_DEPRECATION_WARNINGS
@@ -233,8 +197,6 @@ FRigUnit_AnimNextWriteSkeletalMeshComponentPose_Execute()
 	// Convert and dispatch to renderer
 	UE::Anim::FSkinnedMeshComponentExtensions::CompleteAndDispatch(
 		OutputComponent,
-		RefPose.GetMeshBoneIndexToParentMeshBoneIndexMap(),
-		PoseBoneToMeshBoneIndexMap, // Based on the LOD level of the skeletal mesh component
-		LocalSpaceTransforms,
+		Pose.LODPose,
 		MoveTemp(GameThreadCallback));
 }
