Add GetSocketTransform rig unit

diff --git a/Engine/Plugins/Experimental/AnimNext/Source/AnimNext/Private/Graph/RigUnit_GetSocketTransform.cpp b/Engine/Plugins/Experimental/AnimNext/Source/AnimNext/Private/Graph/RigUnit_GetSocketTransform.cpp
new file mode 100644
--- /dev/null
+++ b/Engine/Plugins/Experimental/AnimNext/Source/AnimNext/Private/Graph/RigUnit_GetSocketTransform.cpp
@@ -0,0 +1,13 @@
+// Copyright Epic Games, Inc. All Rights Reserved.
+
+#include "RigUnit_GetSocketTransform.h"
+
+#include UE_INLINE_GENERATED_CPP_BY_NAME(RigUnit_GetSocketTransform)
+
+FRigUnit_GetSocketTransform_Execute()
+{
+	if (SceneComponent)
+	{
+		Result = SceneComponent->GetSocketTransform(SocketName, TransformSpace);
+	}
+}
diff --git a/Engine/Plugins/Experimental/AnimNext/Source/AnimNext/Private/Graph/RigUnit_GetSocketTransform.h b/Engine/Plugins/Experimental/AnimNext/Source/AnimNext/Private/Graph/RigUnit_GetSocketTransform.h
new file mode 100644
--- /dev/null
+++ b/Engine/Plugins/Experimental/AnimNext/Source/AnimNext/Private/Graph/RigUnit_GetSocketTransform.h
@@ -0,0 +1,31 @@
+// Copyright Epic Games, Inc. All Rights Reserved.
+
+#pragma once
+
+#include "Graph/RigUnit_AnimNextBase.h"
+#include "Components/SceneComponent.h"
+
+#include "RigUnit_GetSocketTransform.generated.h"
+
+// Returns the transform for the specified socket. Note that this is not thread-safe, so must be
+// called at a point where you know that the target socket is not being updated.
+USTRUCT(meta=(DisplayName="Get Socket Transform", Category="Animation Graph", NodeColor="0, 1, 1", Keywords="Socket, Bone, Transform"))
+struct FRigUnit_GetSocketTransform : public FRigUnit_AnimNextBase
+{
+	GENERATED_BODY()
+
+	RIGVM_METHOD()
+	void Execute();
+
+	UPROPERTY(EditAnywhere, Category = "Graph", meta = (Input))
+	TObjectPtr<USceneComponent> SceneComponent = nullptr;
+
+	UPROPERTY(EditAnywhere, Category = "Graph", meta = (Input))
+	FName SocketName = NAME_None;
+
+	UPROPERTY(EditAnywhere, Category = "Graph", meta = (Input))
+	TEnumAsByte<ERelativeTransformSpace> TransformSpace = ERelativeTransformSpace::RTS_World;
+
+	UPROPERTY(EditAnywhere, Category = "Graph", meta = (Output))
+	FTransform Result = FTransform::Identity;
+};
