Improve recently rendered test to match what rendering does

diff --git a/Engine/Plugins/Experimental/AnimNext/Source/AnimNext/Private/Graph/RigUnit_IsRecentlyRendered.cpp b/Engine/Plugins/Experimental/AnimNext/Source/AnimNext/Private/Graph/RigUnit_IsRecentlyRendered.cpp
--- a/Engine/Plugins/Experimental/AnimNext/Source/AnimNext/Private/Graph/RigUnit_IsRecentlyRendered.cpp
+++ b/Engine/Plugins/Experimental/AnimNext/Source/AnimNext/Private/Graph/RigUnit_IsRecentlyRendered.cpp
@@ -1,13 +1,25 @@
 // Copyright Epic Games, Inc. All Rights Reserved.
 
 #include "RigUnit_IsRecentlyRendered.h"
-#include "Components/SkinnedMeshComponent.h"
+#include "SkeletalRenderPublic.h"
 
 FRigUnit_IsRecentlyRendered_Execute()
 {
 	if (const USkinnedMeshComponent* Component = MeshComponent)
 	{
-		Result = Component->bRecentlyRendered;
+		if (const FSkeletalMeshObject* MeshObject = Component->GetMeshObject())
+		{
+			// We are considered recently rendered if we are, or we are set to always tick & refresh, or we haven't been updated yet
+			Result =
+				Component->bRecentlyRendered ||
+				// TODO: Enable this once the content is ready
+				//Component->VisibilityBasedAnimTickOption == EVisibilityBasedAnimTickOption::AlwaysTickPoseAndRefreshBones ||
+				!MeshObject->bHasBeenUpdatedAtLeastOnce;
+		}
+		else
+		{
+			Result = false;
+		}
 	}
 	else
 	{
