Verify user can Undo/Redo in AnimNext Animation Graph
Automated QMetry Test Case: UE-TC-18823

diff --git a/Engine/Plugins/Experimental/AnimNext/Source/AnimNextUncookedOnly/Internal/AnimNextRigVMAssetEditorData.h b/Engine/Plugins/Experimental/AnimNext/Source/AnimNextUncookedOnly/Internal/AnimNextRigVMAssetEditorData.h
--- a/Engine/Plugins/Experimental/AnimNext/Source/AnimNextUncookedOnly/Internal/AnimNextRigVMAssetEditorData.h
+++ b/Engine/Plugins/Experimental/AnimNext/Source/AnimNextUncookedOnly/Internal/AnimNextRigVMAssetEditorData.h
@@ -206,6 +206,7 @@ protected:
 	friend UE::AnimNext::Editor::FAssetCompilationHandler;
 	friend class UAnimNextStateTreeTreeEditorData;
 	friend struct FAnimNextModuleTests;
+	friend struct FAnimNextAnimationGraphTests;
 
 	// UObject interface
 	virtual void Serialize(FArchive& Ar) override;
