Rename AnimNext to UAF in the user interface
[RN] Renamed the experimental AnimNext plugins to UAF in the user interface

diff --git a/Engine/Plugins/Experimental/AnimNext/AnimNext.uplugin b/Engine/Plugins/Experimental/AnimNext/AnimNext.uplugin
--- a/Engine/Plugins/Experimental/AnimNext/AnimNext.uplugin
+++ b/Engine/Plugins/Experimental/AnimNext/AnimNext.uplugin
@@ -2,7 +2,7 @@
 	"FileVersion" : 3,
 	"Version" : 1,
 	"VersionName" : "0.1",
-	"FriendlyName" : "Anim Next",
+	"FriendlyName" : "Unreal Animation Framework (UAF)",
 	"Description" : "Framework for defining functional data flow for animation systems",
 	"Category" : "Animation",
 	"CreatedBy" : "Epic Games, Inc.",
diff --git a/Engine/Plugins/Experimental/AnimNext/Source/AnimNext/Private/AnimNextModuleImpl.cpp b/Engine/Plugins/Experimental/AnimNext/Source/AnimNext/Private/AnimNextModuleImpl.cpp
--- a/Engine/Plugins/Experimental/AnimNext/Source/AnimNext/Private/AnimNextModuleImpl.cpp
+++ b/Engine/Plugins/Experimental/AnimNext/Source/AnimNext/Private/AnimNextModuleImpl.cpp
@@ -84,27 +84,27 @@ namespace UE::AnimNext
 			[&UolModule]
 			{
 				{
-					UE::UniversalObjectLocator::FFragmentTypeParameters FragmentTypeParams("animobjfunc", LOCTEXT("AnimNextObjectFunctionFragment", "Function"));
+					UE::UniversalObjectLocator::FFragmentTypeParameters FragmentTypeParams("animobjfunc", LOCTEXT("UAFObjectFunctionFragment", "Function"));
 					FragmentTypeParams.PrimaryEditorType = "AnimNextObjectFunction";
 					FAnimNextObjectFunctionLocatorFragment::FragmentType = UolModule.RegisterFragmentType<FAnimNextObjectFunctionLocatorFragment>(FragmentTypeParams);
 				}
 				{
-					UE::UniversalObjectLocator::FFragmentTypeParameters FragmentTypeParams("animobjprop", LOCTEXT("AnimNextObjectPropertyFragment", "Property"));
+					UE::UniversalObjectLocator::FFragmentTypeParameters FragmentTypeParams("animobjprop", LOCTEXT("UAFObjectPropertyFragment", "Property"));
 					FragmentTypeParams.PrimaryEditorType = "AnimNextObjectProperty";
 					FAnimNextObjectPropertyLocatorFragment::FragmentType = UolModule.RegisterFragmentType<FAnimNextObjectPropertyLocatorFragment>(FragmentTypeParams);
 				}
 				{
-					UE::UniversalObjectLocator::FFragmentTypeParameters FragmentTypeParams("animobjcast", LOCTEXT("AnimNextCastFragment", "Cast"));
+					UE::UniversalObjectLocator::FFragmentTypeParameters FragmentTypeParams("animobjcast", LOCTEXT("UAFCastFragment", "Cast"));
 					FragmentTypeParams.PrimaryEditorType = "AnimNextObjectCast";
 					FAnimNextObjectCastLocatorFragment::FragmentType = UolModule.RegisterFragmentType<FAnimNextObjectCastLocatorFragment>(FragmentTypeParams);
 				}
 				{
-					UE::UniversalObjectLocator::FFragmentTypeParameters FragmentTypeParams("animcomp", LOCTEXT("AnimNextComponentFragment", "AnimNextComponent"));
+					UE::UniversalObjectLocator::FFragmentTypeParameters FragmentTypeParams("animcomp", LOCTEXT("UAFComponentFragment", "UAFComponent"));
 					FragmentTypeParams.PrimaryEditorType = "AnimNextComponent";
 					FAnimNextComponentLocatorFragment::FragmentType = UolModule.RegisterFragmentType<FAnimNextComponentLocatorFragment>(FragmentTypeParams);
 				}
 				{
-					UE::UniversalObjectLocator::FFragmentTypeParameters FragmentTypeParams("animactor", LOCTEXT("AnimNextActorFragment", "AnimNextActor"));
+					UE::UniversalObjectLocator::FFragmentTypeParameters FragmentTypeParams("animactor", LOCTEXT("UAFActorFragment", "UAFActor"));
 					FragmentTypeParams.PrimaryEditorType = "AnimNextActor";
 					FAnimNextActorLocatorFragment::FragmentType = UolModule.RegisterFragmentType<FAnimNextActorLocatorFragment>(FragmentTypeParams);
 				}
diff --git a/Engine/Plugins/Experimental/AnimNext/Source/AnimNext/Private/Component/AnimNextComponent.cpp b/Engine/Plugins/Experimental/AnimNext/Source/AnimNext/Private/Component/AnimNextComponent.cpp
--- a/Engine/Plugins/Experimental/AnimNext/Source/AnimNext/Private/Component/AnimNextComponent.cpp
+++ b/Engine/Plugins/Experimental/AnimNext/Source/AnimNext/Private/Component/AnimNextComponent.cpp
@@ -122,7 +122,7 @@ DEFINE_FUNCTION(UAnimNextComponent::execSetVariable)
 	{
 		FBlueprintExceptionInfo ExceptionInfo(
 			EBlueprintExceptionType::AbortExecution,
-			NSLOCTEXT("AnimNextComponent", "AnimNextComponent_SetVariableError", "Failed to resolve the Value for Set Variable")
+			NSLOCTEXT("UAFComponent", "UAFComponent_SetVariableError", "Failed to resolve the Value for Set Variable")
 		);
 
 		FBlueprintCoreDelegates::ThrowScriptException(P_THIS, Stack, ExceptionInfo);
@@ -133,7 +133,7 @@ DEFINE_FUNCTION(UAnimNextComponent::execSetVariable)
 	{
 		FBlueprintExceptionInfo ExceptionInfo(
 			EBlueprintExceptionType::NonFatalError,
-			NSLOCTEXT("AnimNextComponent", "AnimNextComponent_SetVariableInvalidWarning", "Invalid variable name supplied to Set Variable")
+			NSLOCTEXT("UAFComponent", "UAFComponent_SetVariableInvalidWarning", "Invalid variable name supplied to Set Variable")
 		);
 
 		FBlueprintCoreDelegates::ThrowScriptException(P_THIS, Stack, ExceptionInfo);
@@ -145,7 +145,7 @@ DEFINE_FUNCTION(UAnimNextComponent::execSetVariable)
 	{
 		FBlueprintExceptionInfo ExceptionInfo(
 			EBlueprintExceptionType::NonFatalError,
-			FText::Format(NSLOCTEXT("AnimNextComponent", "AnimNextComponent_SetVariableNotFoundWarning", "Unknown variable name '{0}' supplied to Set Variable"), FText::FromName(Name))
+			FText::Format(NSLOCTEXT("UAFComponent", "UAFComponent_SetVariableNotFoundWarning", "Unknown variable name '{0}' supplied to Set Variable"), FText::FromName(Name))
 		);
 
 		FBlueprintCoreDelegates::ThrowScriptException(P_THIS, Stack, ExceptionInfo);
@@ -190,7 +190,7 @@ DEFINE_FUNCTION(UAnimNextComponent::execSetVariable)
 	{
 		FBlueprintExceptionInfo ExceptionInfo(
 			EBlueprintExceptionType::NonFatalError,
-			FText::Format(NSLOCTEXT("AnimNextComponent", "AnimNextComponent_SetVariableTypeMismatch", "Incompatible type supplied for variable '{0}'"), FText::FromName(Name))
+			FText::Format(NSLOCTEXT("UAFComponent", "UAFComponent_SetVariableTypeMismatch", "Incompatible type supplied for variable '{0}'"), FText::FromName(Name))
 		);
 
 		FBlueprintCoreDelegates::ThrowScriptException(P_THIS, Stack, ExceptionInfo);
@@ -986,11 +986,11 @@ void UAnimNextComponent::AddModuleEventPrerequisite(FName InEventName, UAnimNext
 {
 	if (!OtherAnimNextComponent)
 	{
-		UE_LOG(LogAnimation, Warning, TEXT("UAnimNextComponent::AddModuleEventPrerequisite called with null OtherAnimNextComponent"));
+		UE_LOG(LogAnimation, Warning, TEXT("UAFComponent::AddModuleEventPrerequisite called with null OtherAnimNextComponent"));
 	}
 	else if (OtherAnimNextComponent == this)
 	{
-		UE_LOG(LogAnimation, Warning, TEXT("UAnimNextComponent::AddModuleEventPrerequisite called using the same component"));
+		UE_LOG(LogAnimation, Warning, TEXT("UAFComponent::AddModuleEventPrerequisite called using the same component"));
 	}
 	else if (Subsystem)
 	{
@@ -1002,11 +1002,11 @@ void UAnimNextComponent::AddModuleEventSubsequent(FName InEventName, UAnimNextCo
 {
 	if (!OtherAnimNextComponent)
 	{
-		UE_LOG(LogAnimation, Warning, TEXT("UAnimNextComponent::AddModuleEventSubsequent called with null OtherAnimNextComponent"));
+		UE_LOG(LogAnimation, Warning, TEXT("UAFComponent::AddModuleEventSubsequent called with null OtherAnimNextComponent"));
 	}
 	else if (OtherAnimNextComponent == this)
 	{
-		UE_LOG(LogAnimation, Warning, TEXT("UAnimNextComponent::AddModuleEventPrerequisite called using the same component"));
+		UE_LOG(LogAnimation, Warning, TEXT("UAFComponent::AddModuleEventPrerequisite called using the same component"));
 	}
 	else if (Subsystem)
 	{
diff --git a/Engine/Plugins/Experimental/AnimNext/Source/AnimNext/Private/GenerationTools.cpp b/Engine/Plugins/Experimental/AnimNext/Source/AnimNext/Private/GenerationTools.cpp
--- a/Engine/Plugins/Experimental/AnimNext/Source/AnimNext/Private/GenerationTools.cpp
+++ b/Engine/Plugins/Experimental/AnimNext/Source/AnimNext/Private/GenerationTools.cpp
@@ -950,7 +950,7 @@ namespace AnimNext::Tools::ConsoleCommands
 	};
 
 
-	FAutoConsoleCommand CheckSkeletalMeshesLODs(TEXT("animnext.tools.checkskeletalmesheslods"), TEXT(""), FConsoleCommandDelegate::CreateLambda([]()
+	FAutoConsoleCommand CheckSkeletalMeshesLODs(TEXT("uaf.tools.checkskeletalmesheslods"), TEXT(""), FConsoleCommandDelegate::CreateLambda([]()
 	{
 		FHelper::CheckSkeletalMeshesLODs();
 	}));
diff --git a/Engine/Plugins/Experimental/AnimNext/Source/AnimNext/Private/Graph/RigUnit_AnimNextRemapPose.h b/Engine/Plugins/Experimental/AnimNext/Source/AnimNext/Private/Graph/RigUnit_AnimNextRemapPose.h
--- a/Engine/Plugins/Experimental/AnimNext/Source/AnimNext/Private/Graph/RigUnit_AnimNextRemapPose.h
+++ b/Engine/Plugins/Experimental/AnimNext/Source/AnimNext/Private/Graph/RigUnit_AnimNextRemapPose.h
@@ -10,7 +10,7 @@
 
 #include "RigUnit_AnimNextRemapPose.generated.h"
 
-DECLARE_CYCLE_STAT_EXTERN(TEXT("AnimNext RigUnit: Remap Pose"), STAT_AnimNext_RigUnit_RemapPose, STATGROUP_AnimNext, ANIMNEXT_API);
+DECLARE_CYCLE_STAT_EXTERN(TEXT("UAF RigUnit: Remap Pose"), STAT_AnimNext_RigUnit_RemapPose, STATGROUP_AnimNext, ANIMNEXT_API);
 
 /** Remaps an anim graph pose from to another skeletal mesh component. */
 USTRUCT(meta=(DisplayName="Remap Pose", Category="Animation Graph", NodeColor="0, 1, 1", Keywords="Output,Pose,Port"))
diff --git a/Engine/Plugins/Experimental/AnimNext/Source/AnimNext/Private/Graph/RigUnit_AnimNextWriteSkeletalMeshComponentPose.cpp b/Engine/Plugins/Experimental/AnimNext/Source/AnimNext/Private/Graph/RigUnit_AnimNextWriteSkeletalMeshComponentPose.cpp
--- a/Engine/Plugins/Experimental/AnimNext/Source/AnimNext/Private/Graph/RigUnit_AnimNextWriteSkeletalMeshComponentPose.cpp
+++ b/Engine/Plugins/Experimental/AnimNext/Source/AnimNext/Private/Graph/RigUnit_AnimNextWriteSkeletalMeshComponentPose.cpp
@@ -54,7 +54,7 @@ FRigUnit_AnimNextWriteSkeletalMeshComponentPose_Execute()
 	// We cannot write to the skeletal mesh component if it is using anim BP
 	if(OutputComponent->bEnableAnimation)
 	{
-		UE_LOGFMT(LogAnimation, Warning, "AnimNext: Could not write to skeletal mesh component - bEnableAnimation is true [SK: {SkeletalMesh}][Owner: {Owner}]",
+		UE_LOGFMT(LogAnimation, Warning, "UAF: Could not write to skeletal mesh component - bEnableAnimation is true [SK: {SkeletalMesh}][Owner: {Owner}]",
 			*SkeletalMesh->GetName(), OutputComponent->GetOuter() != nullptr ? *OutputComponent->GetOuter()->GetName() : TEXT("*NO OWNER*"));
 		return;
 	}
@@ -66,7 +66,7 @@ FRigUnit_AnimNextWriteSkeletalMeshComponentPose_Execute()
 	// You should remap the pose explicitly if this is what you need
 	if (&RefPose != &Pose.LODPose.GetRefPose())
 	{
-		UE_LOGFMT(LogAnimation, Warning, "AnimNext: Could not write to skeletal mesh component - The input pose and the skeletal mesh component use different reference poses [SK: {SkeletalMesh}][Owner: {Owner}]",
+		UE_LOGFMT(LogAnimation, Warning, "UAF: Could not write to skeletal mesh component - The input pose and the skeletal mesh component use different reference poses [SK: {SkeletalMesh}][Owner: {Owner}]",
 			*SkeletalMesh->GetName(), OutputComponent->GetOuter() != nullptr ? *OutputComponent->GetOuter()->GetName() : TEXT("*NO OWNER*"));
 		return;
 	}
diff --git a/Engine/Plugins/Experimental/AnimNext/Source/AnimNext/Private/Module/RigUnit_AnimNextModuleEventDependencies.cpp b/Engine/Plugins/Experimental/AnimNext/Source/AnimNext/Private/Module/RigUnit_AnimNextModuleEventDependencies.cpp
--- a/Engine/Plugins/Experimental/AnimNext/Source/AnimNext/Private/Module/RigUnit_AnimNextModuleEventDependencies.cpp
+++ b/Engine/Plugins/Experimental/AnimNext/Source/AnimNext/Private/Module/RigUnit_AnimNextModuleEventDependencies.cpp
@@ -107,7 +107,7 @@ TArray<FRigVMUserWorkflow> FRigUnit_AnimNextModuleDependenciesBase::GetSupported
 				return false;
 			}
 
-			static const FName MenuName = "AnimNextModuleDependenciesAddMenu";
+			static const FName MenuName = "UAFModuleDependenciesAddMenu";
 			if(!UToolMenus::Get()->IsMenuRegistered(MenuName))
 			{
 				UToolMenu* Menu = UToolMenus::Get()->RegisterMenu(MenuName);
@@ -123,7 +123,7 @@ TArray<FRigVMUserWorkflow> FRigUnit_AnimNextModuleDependenciesBase::GetSupported
 					{
 						FSlateApplication::Get().DismissAllMenus();
 
-						FScopedTransaction Transaction(NSLOCTEXT("AnimNextModuleDependencies", "AddDependencyTraitTransaction", "Add dependency trait"));
+						FScopedTransaction Transaction(NSLOCTEXT("UAFModuleDependencies", "AddDependencyTraitTransaction", "Add dependency trait"));
 						Controller->AddTrait(Node, const_cast<UScriptStruct*>(InStruct), InStruct->GetFName(), TEXT(""));
 					};
 
@@ -212,7 +212,7 @@ TArray<FRigVMUserWorkflow> FRigUnit_AnimNextModuleDependenciesBase::GetSupported
 				return false;
 			}
 
-			FScopedTransaction Transaction(NSLOCTEXT("AnimNextModuleDependencies", "RemoveDependencyTraitTransaction", "Remove dependency trait"));
+			FScopedTransaction Transaction(NSLOCTEXT("UAFModuleDependencies", "RemoveDependencyTraitTransaction", "Remove dependency trait"));
 			Controller->RemoveTrait(Pin->GetNode(), TraitPin->GetFName());
 
 			return true;
diff --git a/Engine/Plugins/Experimental/AnimNext/Source/AnimNext/Private/Param/AnimNextObjectPropertyLocatorFragment.cpp b/Engine/Plugins/Experimental/AnimNext/Source/AnimNext/Private/Param/AnimNextObjectPropertyLocatorFragment.cpp
--- a/Engine/Plugins/Experimental/AnimNext/Source/AnimNext/Private/Param/AnimNextObjectPropertyLocatorFragment.cpp
+++ b/Engine/Plugins/Experimental/AnimNext/Source/AnimNext/Private/Param/AnimNextObjectPropertyLocatorFragment.cpp
@@ -8,7 +8,7 @@
 #include "UniversalObjectLocatorInitializeResult.h"
 #include "String/ParseTokens.h"
 
-#define LOCTEXT_NAMESPACE "AnimNextObjectPropertyLocatorFragment"
+#define LOCTEXT_NAMESPACE "UAFObjectPropertyLocatorFragment"
 
 UE::UniversalObjectLocator::TFragmentTypeHandle<FAnimNextObjectPropertyLocatorFragment> FAnimNextObjectPropertyLocatorFragment::FragmentType;
 
diff --git a/Engine/Plugins/Experimental/AnimNext/Source/AnimNext/Private/RewindDebugger/RewindDebuggerAnimNextRuntime.cpp b/Engine/Plugins/Experimental/AnimNext/Source/AnimNext/Private/RewindDebugger/RewindDebuggerAnimNextRuntime.cpp
--- a/Engine/Plugins/Experimental/AnimNext/Source/AnimNext/Private/RewindDebugger/RewindDebuggerAnimNextRuntime.cpp
+++ b/Engine/Plugins/Experimental/AnimNext/Source/AnimNext/Private/RewindDebugger/RewindDebuggerAnimNextRuntime.cpp
@@ -9,12 +9,12 @@ namespace UE::AnimNext
 	
 void FRewindDebuggerAnimNextRuntime::RecordingStarted()
 {
-	UE::Trace::ToggleChannel(TEXT("AnimNext"), true);
+	UE::Trace::ToggleChannel(TEXT("UAF"), true);
 }
 
 void FRewindDebuggerAnimNextRuntime::RecordingStopped()
 {
-	UE::Trace::ToggleChannel(TEXT("AnimNext"), false);
+	UE::Trace::ToggleChannel(TEXT("UAF"), false);
 }
 
 void FRewindDebuggerAnimNextRuntime::Clear()
diff --git a/Engine/Plugins/Experimental/AnimNext/Source/AnimNext/Public/AnimNextStats.h b/Engine/Plugins/Experimental/AnimNext/Source/AnimNext/Public/AnimNextStats.h
--- a/Engine/Plugins/Experimental/AnimNext/Source/AnimNext/Public/AnimNextStats.h
+++ b/Engine/Plugins/Experimental/AnimNext/Source/AnimNext/Public/AnimNextStats.h
@@ -4,25 +4,25 @@
 
 #include "Stats/Stats.h"
 
-DECLARE_STATS_GROUP(TEXT("AnimNext"), STATGROUP_AnimNext, STATCAT_Advanced);
+DECLARE_STATS_GROUP(TEXT("UAF"), STATGROUP_AnimNext, STATCAT_Advanced);
 
-DECLARE_CYCLE_STAT_EXTERN(TEXT("AnimNext: Module Event Tick"), STAT_AnimNext_Module_EventTick, STATGROUP_AnimNext, ANIMNEXT_API);
-DECLARE_CYCLE_STAT_EXTERN(TEXT("AnimNext: Module End Tick"), STAT_AnimNext_Module_EndTick, STATGROUP_AnimNext, ANIMNEXT_API);
+DECLARE_CYCLE_STAT_EXTERN(TEXT("UAF: Module Event Tick"), STAT_AnimNext_Module_EventTick, STATGROUP_AnimNext, ANIMNEXT_API);
+DECLARE_CYCLE_STAT_EXTERN(TEXT("UAF: Module End Tick"), STAT_AnimNext_Module_EndTick, STATGROUP_AnimNext, ANIMNEXT_API);
 
-DECLARE_CYCLE_STAT_EXTERN(TEXT("AnimNext: Initialize Instance"), STAT_AnimNext_InitializeInstance, STATGROUP_AnimNext, ANIMNEXT_API);
-DECLARE_CYCLE_STAT_EXTERN(TEXT("AnimNext: Create Instance Data"), STAT_AnimNext_CreateInstanceData, STATGROUP_AnimNext, ANIMNEXT_API);
+DECLARE_CYCLE_STAT_EXTERN(TEXT("UAF: Initialize Instance"), STAT_AnimNext_InitializeInstance, STATGROUP_AnimNext, ANIMNEXT_API);
+DECLARE_CYCLE_STAT_EXTERN(TEXT("UAF: Create Instance Data"), STAT_AnimNext_CreateInstanceData, STATGROUP_AnimNext, ANIMNEXT_API);
 
-DECLARE_CYCLE_STAT_EXTERN(TEXT("AnimNext: Make RefPose"), STAT_AnimNext_Make_RefPose, STATGROUP_AnimNext, ANIMNEXT_API);
-DECLARE_CYCLE_STAT_EXTERN(TEXT("AnimNext: Write Pose"), STAT_AnimNext_Write_Pose, STATGROUP_AnimNext, ANIMNEXT_API);
+DECLARE_CYCLE_STAT_EXTERN(TEXT("UAF: Make RefPose"), STAT_AnimNext_Make_RefPose, STATGROUP_AnimNext, ANIMNEXT_API);
+DECLARE_CYCLE_STAT_EXTERN(TEXT("UAF: Write Pose"), STAT_AnimNext_Write_Pose, STATGROUP_AnimNext, ANIMNEXT_API);
 
-DECLARE_CYCLE_STAT_EXTERN(TEXT("AnimNext: Copy Transforms (SoA)"), STAT_AnimNext_CopyTransforms_SoA, STATGROUP_AnimNext, ANIMNEXT_API);
-DECLARE_CYCLE_STAT_EXTERN(TEXT("AnimNext: Normalize Rotations (SoA)"), STAT_AnimNext_NormalizeRotations_SoA, STATGROUP_AnimNext, ANIMNEXT_API);
-DECLARE_CYCLE_STAT_EXTERN(TEXT("AnimNext: Blend Overwrite (SoA)"), STAT_AnimNext_BlendOverwrite_SoA, STATGROUP_AnimNext, ANIMNEXT_API);
-DECLARE_CYCLE_STAT_EXTERN(TEXT("AnimNext: Blend Accumulate (SoA)"), STAT_AnimNext_BlendAccumulate_SoA, STATGROUP_AnimNext, ANIMNEXT_API);
+DECLARE_CYCLE_STAT_EXTERN(TEXT("UAF: Copy Transforms (SoA)"), STAT_AnimNext_CopyTransforms_SoA, STATGROUP_AnimNext, ANIMNEXT_API);
+DECLARE_CYCLE_STAT_EXTERN(TEXT("UAF: Normalize Rotations (SoA)"), STAT_AnimNext_NormalizeRotations_SoA, STATGROUP_AnimNext, ANIMNEXT_API);
+DECLARE_CYCLE_STAT_EXTERN(TEXT("UAF: Blend Overwrite (SoA)"), STAT_AnimNext_BlendOverwrite_SoA, STATGROUP_AnimNext, ANIMNEXT_API);
+DECLARE_CYCLE_STAT_EXTERN(TEXT("UAF: Blend Accumulate (SoA)"), STAT_AnimNext_BlendAccumulate_SoA, STATGROUP_AnimNext, ANIMNEXT_API);
 
-DECLARE_CYCLE_STAT_EXTERN(TEXT("AnimNext: Generate Reference Pose"), STAT_AnimNext_GenerateReferencePose, STATGROUP_AnimNext, ANIMNEXT_API);
-DECLARE_CYCLE_STAT_EXTERN(TEXT("AnimNext: Remap Pose (From AnimBP)"), STAT_AnimNext_RemapPose_FromAnimBP, STATGROUP_AnimNext, ANIMNEXT_API);
-DECLARE_CYCLE_STAT_EXTERN(TEXT("AnimNext: Remap Pose (To AnimBP)"), STAT_AnimNext_RemapPose_ToAnimBP, STATGROUP_AnimNext, ANIMNEXT_API);
-DECLARE_CYCLE_STAT_EXTERN(TEXT("AnimNext: Remap Pose (To LocalTransforms)"), STAT_AnimNext_RemapPose_ToLocalTransforms, STATGROUP_AnimNext, ANIMNEXT_API);
-DECLARE_CYCLE_STAT_EXTERN(TEXT("AnimNext: Remap Pose (Mesh->Mesh)"), STAT_AnimNext_RemapPose_Mesh2Mesh, STATGROUP_AnimNext, ANIMNEXT_API);
-DECLARE_CYCLE_STAT_EXTERN(TEXT("AnimNext: Convert Local Space To Component Space"), STAT_AnimNext_ConvertLocalSpaceToComponentSpace, STATGROUP_AnimNext, ANIMNEXT_API);
\ No newline at end of file
+DECLARE_CYCLE_STAT_EXTERN(TEXT("UAF: Generate Reference Pose"), STAT_AnimNext_GenerateReferencePose, STATGROUP_AnimNext, ANIMNEXT_API);
+DECLARE_CYCLE_STAT_EXTERN(TEXT("UAF: Remap Pose (From AnimBP)"), STAT_AnimNext_RemapPose_FromAnimBP, STATGROUP_AnimNext, ANIMNEXT_API);
+DECLARE_CYCLE_STAT_EXTERN(TEXT("UAF: Remap Pose (To AnimBP)"), STAT_AnimNext_RemapPose_ToAnimBP, STATGROUP_AnimNext, ANIMNEXT_API);
+DECLARE_CYCLE_STAT_EXTERN(TEXT("UAF: Remap Pose (To LocalTransforms)"), STAT_AnimNext_RemapPose_ToLocalTransforms, STATGROUP_AnimNext, ANIMNEXT_API);
+DECLARE_CYCLE_STAT_EXTERN(TEXT("UAF: Remap Pose (Mesh->Mesh)"), STAT_AnimNext_RemapPose_Mesh2Mesh, STATGROUP_AnimNext, ANIMNEXT_API);
+DECLARE_CYCLE_STAT_EXTERN(TEXT("UAF: Convert Local Space To Component Space"), STAT_AnimNext_ConvertLocalSpaceToComponentSpace, STATGROUP_AnimNext, ANIMNEXT_API);
\ No newline at end of file
diff --git a/Engine/Plugins/Experimental/AnimNext/Source/AnimNext/Public/Component/AnimNextComponent.h b/Engine/Plugins/Experimental/AnimNext/Source/AnimNext/Public/Component/AnimNextComponent.h
--- a/Engine/Plugins/Experimental/AnimNext/Source/AnimNext/Public/Component/AnimNextComponent.h
+++ b/Engine/Plugins/Experimental/AnimNext/Source/AnimNext/Public/Component/AnimNextComponent.h
@@ -57,7 +57,7 @@ public:
 	// Sets a module variable's value.
 	// @param    Name     The name of the variable to set
 	// @param    Value    The value to set the variable to
-	UFUNCTION(BlueprintCallable, Category = "AnimNext", CustomThunk, meta = (CustomStructureParam = Value, UnsafeDuringActorConstruction))
+	UFUNCTION(BlueprintCallable, Category = "UAF", CustomThunk, meta = (CustomStructureParam = Value, UnsafeDuringActorConstruction))
 	ANIMNEXT_API void SetVariable(UPARAM(meta = (CustomWidget = "VariableName")) FName Name, int32 Value);
 
 	/**
@@ -184,15 +184,15 @@ public:
 	}
 
 	// Whether this component is currently updating
-	UFUNCTION(BlueprintCallable, Category = "AnimNext")
+	UFUNCTION(BlueprintCallable, Category = "UAF")
 	ANIMNEXT_API bool IsEnabled() const;
 
 	// Enable or disable this component's update
-	UFUNCTION(BlueprintCallable, Category = "AnimNext")
+	UFUNCTION(BlueprintCallable, Category = "UAF")
 	ANIMNEXT_API void SetEnabled(bool bEnabled);
 
 	// Enable or disable debug drawing. Note only works in builds with UE_ENABLE_DEBUG_DRAWING enabled
-	UFUNCTION(BlueprintCallable, Category = "AnimNext")
+	UFUNCTION(BlueprintCallable, Category = "UAF")
 	ANIMNEXT_API void ShowDebugDrawing(bool bShowDebugDrawing);
 
 	// Queue a task to run during execution 
@@ -219,7 +219,7 @@ public:
 	// The component will tick before the event
 	// @param	Component			The component to add as a prerequisite 
 	// @param	EventName			The event to add the dependency to
-	UFUNCTION(BlueprintCallable, Category = "AnimNext")
+	UFUNCTION(BlueprintCallable, Category = "UAF")
 	ANIMNEXT_API void AddComponentPrerequisite(UActorComponent* Component, FName EventName);
 
 	// Add a subsequent tick function dependency to the specified event
@@ -232,7 +232,7 @@ public:
 	// The component will tick after the event
 	// @param	Component			The component to add as a subsequent of the event
 	// @param	EventName			The event to add the dependency to
-	UFUNCTION(BlueprintCallable, Category = "AnimNext")
+	UFUNCTION(BlueprintCallable, Category = "UAF")
 	ANIMNEXT_API void AddComponentSubsequent(UActorComponent* Component, FName EventName);
 
 	// Remove a prerequisite tick function dependency from the specified event
@@ -244,7 +244,7 @@ public:
 	// Remove a prerequisite on the component's primary tick function from the specified event
 	// @param	Component			The component to remove as a prerequisite 
 	// @param	EventName			The event to add the dependency to
-	UFUNCTION(BlueprintCallable, Category = "AnimNext")
+	UFUNCTION(BlueprintCallable, Category = "UAF")
 	ANIMNEXT_API void RemoveComponentPrerequisite(UActorComponent* Component, FName EventName);
 	
 	// Remove a prerequisite tick function dependency from the specified event
@@ -256,38 +256,38 @@ public:
 	// Remove a subsequent dependency on the component's primary tick function from the specified event
 	// @param	Component			The component to remove as a subsequent of the event
 	// @param	EventName			The event to add the dependency to
-	UFUNCTION(BlueprintCallable, Category = "AnimNext")
+	UFUNCTION(BlueprintCallable, Category = "UAF")
 	ANIMNEXT_API void RemoveComponentSubsequent(UActorComponent* Component, FName EventName);
 	
 	// Add a prerequisite anim next event dependency to the specified event
 	// @param	InEventName			The event name in this component
 	// @param	InTickFunction		The other component we want a prerequisite on
 	// @param	InEventName			The other component's event name
-	UFUNCTION(BlueprintCallable, Category = "AnimNext")
+	UFUNCTION(BlueprintCallable, Category = "UAF")
 	ANIMNEXT_API void AddModuleEventPrerequisite(FName InEventName, UAnimNextComponent* OtherAnimNextComponent, FName OtherEventName);
 
 	// Add a subsequent anim next event dependency to the specified event
 	// @param	InEventName			The event name in this component
 	// @param	InTickFunction		The other component we want to add a prerequisite to
 	// @param	InEventName			The other component's event name
-	UFUNCTION(BlueprintCallable, Category = "AnimNext")
+	UFUNCTION(BlueprintCallable, Category = "UAF")
 	ANIMNEXT_API void AddModuleEventSubsequent(FName InEventName, UAnimNextComponent* OtherAnimNextComponent, FName OtherEventName);
 
 	// Remove a prerequisite anim next event dependency from the specified event
 	// @param	InEventName			The event name in this component
 	// @param	InTickFunction		The other component we want to remove a prerequisite from
 	// @param	InEventName			The other component's event name
-	UFUNCTION(BlueprintCallable, Category = "AnimNext")
+	UFUNCTION(BlueprintCallable, Category = "UAF")
 	ANIMNEXT_API void RemoveModuleEventPrerequisite(FName InEventName, UAnimNextComponent* OtherAnimNextComponent, FName OtherEventName);
 
 	// Remove a subsequent anim next event dependency from the specified event
 	// @param	InEventName			The event name in this component
 	// @param	InTickFunction		The other component we want to remove a prerequisite to
 	// @param	InEventName			The other component's event name
-	UFUNCTION(BlueprintCallable, Category = "AnimNext")
+	UFUNCTION(BlueprintCallable, Category = "UAF")
 	ANIMNEXT_API void RemoveModuleEventSubsequent(FName InEventName, UAnimNextComponent* OtherAnimNextComponent, FName OtherEventName);
 
-	UFUNCTION(BlueprintCallable, Category = "AnimNext", meta=(DisplayName = "Get Module Handle"))
+	UFUNCTION(BlueprintCallable, Category = "UAF", meta=(DisplayName = "Get Module Handle"))
 	ANIMNEXT_API FAnimNextModuleHandle BlueprintGetModuleHandle() const;
 
 private:
diff --git a/Engine/Plugins/Experimental/AnimNext/Source/AnimNextEditor/Private/AnimNextEditorModule.cpp b/Engine/Plugins/Experimental/AnimNext/Source/AnimNextEditor/Private/AnimNextEditorModule.cpp
--- a/Engine/Plugins/Experimental/AnimNext/Source/AnimNextEditor/Private/AnimNextEditorModule.cpp
+++ b/Engine/Plugins/Experimental/AnimNext/Source/AnimNextEditor/Private/AnimNextEditorModule.cpp
@@ -99,8 +99,8 @@ void FAnimNextEditorModule::StartupModule()
 
 	// Register settings for user editing
 	ISettingsModule& SettingsModule = FModuleManager::Get().LoadModuleChecked<ISettingsModule>("Settings");
-	SettingsModule.RegisterSettings("Editor", "General", "AnimNext",
-		LOCTEXT("SettingsName", "AnimNext"),
+	SettingsModule.RegisterSettings("Editor", "General", "UAF",
+		LOCTEXT("SettingsName", "UAF"),
 		LOCTEXT("SettingsDescription", "Customize AnimNext Settings."),
 		GetMutableDefault<UAnimNextConfig>()
 	);
@@ -326,7 +326,7 @@ void FAnimNextEditorModule::StartupModule()
 	UolEditorModule.RegisterLocatorEditor("AnimNextComponent", MakeShared<FComponentLocatorEditor>());
 	UolEditorModule.RegisterLocatorEditor("AnimNextActor", MakeShared<FActorLocatorEditor>());
 
-	UolEditorModule.RegisterEditorContext("AnimNextContext", MakeShared<FLocatorContext>());
+	UolEditorModule.RegisterEditorContext("UAFContext", MakeShared<FLocatorContext>());
 
 	RegisterLocatorFragmentEditorType("Actor");
 	RegisterLocatorFragmentEditorType("Asset");
@@ -395,7 +395,7 @@ void FAnimNextEditorModule::ShutdownModule()
 		UolEditorModule.UnregisterLocatorEditor("AnimNextComponent");
 		UolEditorModule.UnregisterLocatorEditor("AnimNextActor");
 
-		UolEditorModule.UnregisterEditorContext("AnimNextContext");
+		UolEditorModule.UnregisterEditorContext("UAFContext");
 	}
 	
 	if (UObjectInitialized())
diff --git a/Engine/Plugins/Experimental/AnimNext/Source/AnimNextEditor/Private/Common/AnimNextEntryAssetDefinitions.h b/Engine/Plugins/Experimental/AnimNext/Source/AnimNextEditor/Private/Common/AnimNextEntryAssetDefinitions.h
--- a/Engine/Plugins/Experimental/AnimNext/Source/AnimNextEditor/Private/Common/AnimNextEntryAssetDefinitions.h
+++ b/Engine/Plugins/Experimental/AnimNext/Source/AnimNextEditor/Private/Common/AnimNextEntryAssetDefinitions.h
@@ -1,4 +1,4 @@
-﻿// Copyright Epic Games, Inc. All Rights Reserved.
+// Copyright Epic Games, Inc. All Rights Reserved.
 
 #pragma once
 
@@ -17,7 +17,7 @@ class UAssetDefinition_AnimNextVariableEntry : public UAssetDefinitionDefault
 
 public:
 	// UAssetDefinition interface
-	virtual FText GetAssetDisplayName() const override { return LOCTEXT("AnimNextVariable", "Variable"); }
+	virtual FText GetAssetDisplayName() const override { return LOCTEXT("UAFVariable", "Variable"); }
 	virtual FLinearColor GetAssetColor() const override { return FLinearColor(FColor(128,64,64)); }
 	virtual TSoftClassPtr<UObject> GetAssetClass() const override { return UAnimNextVariableEntry::StaticClass(); }
 	virtual FText GetObjectDisplayNameText(UObject* Object) const override;
@@ -30,7 +30,7 @@ class UAssetDefinition_AnimNextEventGraphEntry : public UAssetDefinitionDefault
 
 public:
 	// UAssetDefinition interface
-	virtual FText GetAssetDisplayName() const override { return LOCTEXT("AnimNextEventGraphEntry", "Event Graph"); }
+	virtual FText GetAssetDisplayName() const override { return LOCTEXT("UAFEventGraphEntry", "Event Graph"); }
 	virtual FLinearColor GetAssetColor() const override { return FLinearColor(FColor(128,64,64)); }
 	virtual TSoftClassPtr<UObject> GetAssetClass() const override { return UAnimNextEventGraphEntry::StaticClass(); }
 	virtual FText GetObjectDisplayNameText(UObject* Object) const override;
@@ -43,7 +43,7 @@ class UAssetDefinition_AnimNextDataInterfaceEntry : public UAssetDefinitionDefau
 
 public:
 	// UAssetDefinition interface
-	virtual FText GetAssetDisplayName() const override { return LOCTEXT("AnimNextDataInterfaceEntry", "Data Interface"); }
+	virtual FText GetAssetDisplayName() const override { return LOCTEXT("UAFDataInterfaceEntry", "Data Interface"); }
 	virtual FLinearColor GetAssetColor() const override { return FLinearColor(FColor(128,64,64)); }
 	virtual TSoftClassPtr<UObject> GetAssetClass() const override { return UAnimNextDataInterfaceEntry::StaticClass(); }
 	virtual FText GetObjectDisplayNameText(UObject* Object) const override;
diff --git a/Engine/Plugins/Experimental/AnimNext/Source/AnimNextEditor/Private/Common/AnimNextFindTabSummoner.cpp b/Engine/Plugins/Experimental/AnimNext/Source/AnimNextEditor/Private/Common/AnimNextFindTabSummoner.cpp
--- a/Engine/Plugins/Experimental/AnimNext/Source/AnimNextEditor/Private/Common/AnimNextFindTabSummoner.cpp
+++ b/Engine/Plugins/Experimental/AnimNext/Source/AnimNextEditor/Private/Common/AnimNextFindTabSummoner.cpp
@@ -17,10 +17,10 @@ namespace UE::AnimNext::Editor
 FAnimNextFindTabSummoner::FAnimNextFindTabSummoner(TSharedPtr<UE::Workspace::IWorkspaceEditor> InHostingApp)
 	: FWorkflowTabFactory(UE::AnimNext::Editor::FindTabName, StaticCastSharedPtr<FAssetEditorToolkit>(InHostingApp))
 {
-	TabLabel = LOCTEXT("AnimNextFindTabLabel", "Find");
+	TabLabel = LOCTEXT("UAFFindTabLabel", "Find");
 	TabIcon = FSlateIcon(FAppStyle::GetAppStyleSetName(), "Kismet.Tabs.FindResults");
-	ViewMenuDescription = LOCTEXT("AnimNextFindTabMenuDescription", "Find");
-	ViewMenuTooltip = LOCTEXT("AnimNextFindTabToolTip", "Search contents of currently selected AnimNext Asset.");
+	ViewMenuDescription = LOCTEXT("UAFFindTabMenuDescription", "Find");
+	ViewMenuTooltip = LOCTEXT("UAFFindTabToolTip", "Search contents of currently selected UAF Asset.");
 	bIsSingleton = true;
 
 	AnimNextFindResultsWidget = SNew(SFindInAnimNextRigVMAsset, InHostingApp);
diff --git a/Engine/Plugins/Experimental/AnimNext/Source/AnimNextEditor/Private/Common/AnimNextRigVMAssetCommands.cpp b/Engine/Plugins/Experimental/AnimNext/Source/AnimNextEditor/Private/Common/AnimNextRigVMAssetCommands.cpp
--- a/Engine/Plugins/Experimental/AnimNext/Source/AnimNextEditor/Private/Common/AnimNextRigVMAssetCommands.cpp
+++ b/Engine/Plugins/Experimental/AnimNext/Source/AnimNextEditor/Private/Common/AnimNextRigVMAssetCommands.cpp
@@ -8,13 +8,13 @@ namespace UE::AnimNext
 {
 
 FAnimNextRigVMAssetCommands::FAnimNextRigVMAssetCommands()
-	: TCommands<FAnimNextRigVMAssetCommands>("AnimNextRigVMAssetCommand", LOCTEXT("AnimNextRigVMAssetCommands", "Anim Next Asset Commands"), NAME_None, "AnimNextStyle")
+	: TCommands<FAnimNextRigVMAssetCommands>("UAFRigVMAssetCommand", LOCTEXT("UAFRigVMAssetCommands", "UAF Asset Commands"), NAME_None, "AnimNextStyle")
 {
 }
 
 void FAnimNextRigVMAssetCommands::RegisterCommands()
 {
-	UI_COMMAND(FindInAnimNextRigVMAsset, "FindInAnimNextRigVMAsset", "Search the current AnimNext Asset.", EUserInterfaceActionType::Button, FInputChord(EModifierKey::Control, EKeys::F));
+	UI_COMMAND(FindInAnimNextRigVMAsset, "FindInUAFRigVMAsset", "Search the current UAF Asset.", EUserInterfaceActionType::Button, FInputChord(EModifierKey::Control, EKeys::F));
 }
 
 }
diff --git a/Engine/Plugins/Experimental/AnimNext/Source/AnimNextEditor/Private/DataInterface/AnimNextDataInterfaceAssetDefinition.h b/Engine/Plugins/Experimental/AnimNext/Source/AnimNextEditor/Private/DataInterface/AnimNextDataInterfaceAssetDefinition.h
--- a/Engine/Plugins/Experimental/AnimNext/Source/AnimNextEditor/Private/DataInterface/AnimNextDataInterfaceAssetDefinition.h
+++ b/Engine/Plugins/Experimental/AnimNext/Source/AnimNextEditor/Private/DataInterface/AnimNextDataInterfaceAssetDefinition.h
@@ -1,4 +1,4 @@
-﻿// Copyright Epic Games, Inc. All Rights Reserved.
+// Copyright Epic Games, Inc. All Rights Reserved.
 
 #pragma once
 
@@ -15,13 +15,13 @@ class UAssetDefinition_AnimNextDataInterface : public UAssetDefinitionDefault
 
 public:
 	// UAssetDefinition interface
-	virtual FText GetAssetDisplayName() const override { return LOCTEXT("AnimNextDataInterface", "AnimNext Data Interface"); }
+	virtual FText GetAssetDisplayName() const override { return LOCTEXT("UAFDataInterface", "UAF Data Interface"); }
 	virtual FLinearColor GetAssetColor() const override { return FLinearColor(FColor(128,64,32)); }
 	virtual TSoftClassPtr<UObject> GetAssetClass() const override { return UAnimNextDataInterface::StaticClass(); }
 	virtual EAssetCommandResult OpenAssets(const FAssetOpenArgs& OpenArgs) const override;
 	virtual TConstArrayView<FAssetCategoryPath> GetAssetCategories() const override
 	{
-		static const auto Categories = { FAssetCategoryPath(EAssetCategoryPaths::Animation, LOCTEXT("AnimNextSubMenu", "AnimNext")) };
+		static const auto Categories = { FAssetCategoryPath(EAssetCategoryPaths::Animation, LOCTEXT("UAFSubMenu", "Animation Framework")) };
 		return Categories;
 	}
 	virtual bool ShouldSaveExternalPackages() const override { return true; }
diff --git a/Engine/Plugins/Experimental/AnimNext/Source/AnimNextEditor/Private/Module/AnimNextModuleAssetDefinition.h b/Engine/Plugins/Experimental/AnimNext/Source/AnimNextEditor/Private/Module/AnimNextModuleAssetDefinition.h
--- a/Engine/Plugins/Experimental/AnimNext/Source/AnimNextEditor/Private/Module/AnimNextModuleAssetDefinition.h
+++ b/Engine/Plugins/Experimental/AnimNext/Source/AnimNextEditor/Private/Module/AnimNextModuleAssetDefinition.h
@@ -15,13 +15,13 @@ class UAssetDefinition_AnimNextModule : public UAssetDefinitionDefault
 
 public:
 	// UAssetDefinition interface
-	virtual FText GetAssetDisplayName() const override { return LOCTEXT("AnimNextModule", "AnimNext Module"); }
+	virtual FText GetAssetDisplayName() const override { return LOCTEXT("UAFModule", "UAF Module"); }
 	virtual FLinearColor GetAssetColor() const override { return FLinearColor(FColor(128,96,48)); }
 	virtual TSoftClassPtr<UObject> GetAssetClass() const override { return UAnimNextModule::StaticClass(); }
 	virtual EAssetCommandResult OpenAssets(const FAssetOpenArgs& OpenArgs) const override;
 	virtual TConstArrayView<FAssetCategoryPath> GetAssetCategories() const override
 	{
-		static const auto Categories = { FAssetCategoryPath(EAssetCategoryPaths::Animation, LOCTEXT("AnimNextSubMenu", "AnimNext")) };
+		static const auto Categories = { FAssetCategoryPath(EAssetCategoryPaths::Animation, LOCTEXT("UAFSubMenu", "Animation Framework")) };
 		return Categories;
 	}
 	virtual bool ShouldSaveExternalPackages() const override { return true; }
diff --git a/Engine/Plugins/Experimental/AnimNext/Source/AnimNextEditor/Private/Param/AnimNextActorLocatorEditor.cpp b/Engine/Plugins/Experimental/AnimNext/Source/AnimNextEditor/Private/Param/AnimNextActorLocatorEditor.cpp
--- a/Engine/Plugins/Experimental/AnimNext/Source/AnimNextEditor/Private/Param/AnimNextActorLocatorEditor.cpp
+++ b/Engine/Plugins/Experimental/AnimNext/Source/AnimNextEditor/Private/Param/AnimNextActorLocatorEditor.cpp
@@ -1,4 +1,4 @@
-﻿// Copyright Epic Games, Inc. All Rights Reserved.
+// Copyright Epic Games, Inc. All Rights Reserved.
 
 #include "AnimNextActorLocatorEditor.h"
 
@@ -21,7 +21,7 @@ UE::UniversalObjectLocator::ELocatorFragmentEditorType FActorLocatorEditor::GetL
 
 bool FActorLocatorEditor::IsAllowedInContext(FName InContextName) const
 {
-	return InContextName == "AnimNextContext";
+	return InContextName == "UAFContext";
 }
 
 bool FActorLocatorEditor::IsDragSupported(TSharedPtr<FDragDropOperation> DragOperation, UObject* Context) const
diff --git a/Engine/Plugins/Experimental/AnimNext/Source/AnimNextEditor/Private/Param/AnimNextComponentLocatorEditor.cpp b/Engine/Plugins/Experimental/AnimNext/Source/AnimNextEditor/Private/Param/AnimNextComponentLocatorEditor.cpp
--- a/Engine/Plugins/Experimental/AnimNext/Source/AnimNextEditor/Private/Param/AnimNextComponentLocatorEditor.cpp
+++ b/Engine/Plugins/Experimental/AnimNext/Source/AnimNextEditor/Private/Param/AnimNextComponentLocatorEditor.cpp
@@ -1,4 +1,4 @@
-﻿// Copyright Epic Games, Inc. All Rights Reserved.
+// Copyright Epic Games, Inc. All Rights Reserved.
 
 #include "AnimNextComponentLocatorEditor.h"
 
@@ -89,7 +89,7 @@ UE::UniversalObjectLocator::ELocatorFragmentEditorType FComponentLocatorEditor::
 
 bool FComponentLocatorEditor::IsAllowedInContext(FName InContextName) const
 {
-	return InContextName == "AnimNextContext";
+	return InContextName == "UAFContext";
 }
 
 bool FComponentLocatorEditor::IsDragSupported(TSharedPtr<FDragDropOperation> DragOperation, UObject* Context) const
diff --git a/Engine/Plugins/Experimental/AnimNext/Source/AnimNextEditor/Private/Param/ObjectCastLocatorEditor.cpp b/Engine/Plugins/Experimental/AnimNext/Source/AnimNextEditor/Private/Param/ObjectCastLocatorEditor.cpp
--- a/Engine/Plugins/Experimental/AnimNext/Source/AnimNextEditor/Private/Param/ObjectCastLocatorEditor.cpp
+++ b/Engine/Plugins/Experimental/AnimNext/Source/AnimNextEditor/Private/Param/ObjectCastLocatorEditor.cpp
@@ -1,4 +1,4 @@
-﻿// Copyright Epic Games, Inc. All Rights Reserved.
+// Copyright Epic Games, Inc. All Rights Reserved.
 
 #include "ObjectCastLocatorEditor.h"
 
@@ -73,7 +73,7 @@ UE::UniversalObjectLocator::ELocatorFragmentEditorType FObjectCastLocatorEditor:
 
 bool FObjectCastLocatorEditor::IsAllowedInContext(FName InContextName) const
 {
-	return InContextName == "AnimNextContext";
+	return InContextName == "UAFContext";
 }
 
 bool FObjectCastLocatorEditor::IsDragSupported(TSharedPtr<FDragDropOperation> DragOperation, UObject* Context) const
diff --git a/Engine/Plugins/Experimental/AnimNext/Source/AnimNextEditor/Private/Param/ObjectFunctionLocatorEditor.cpp b/Engine/Plugins/Experimental/AnimNext/Source/AnimNextEditor/Private/Param/ObjectFunctionLocatorEditor.cpp
--- a/Engine/Plugins/Experimental/AnimNext/Source/AnimNextEditor/Private/Param/ObjectFunctionLocatorEditor.cpp
+++ b/Engine/Plugins/Experimental/AnimNext/Source/AnimNextEditor/Private/Param/ObjectFunctionLocatorEditor.cpp
@@ -154,7 +154,7 @@ UE::UniversalObjectLocator::ELocatorFragmentEditorType FObjectFunctionLocatorEdi
 
 bool FObjectFunctionLocatorEditor::IsAllowedInContext(FName InContextName) const
 {
-	return InContextName == "AnimNextContext";
+	return InContextName == "UAFContext";
 }
 
 bool FObjectFunctionLocatorEditor::IsDragSupported(TSharedPtr<FDragDropOperation> DragOperation, UObject* Context) const
diff --git a/Engine/Plugins/Experimental/AnimNext/Source/AnimNextEditor/Private/Param/ObjectPropertyLocatorEditor.cpp b/Engine/Plugins/Experimental/AnimNext/Source/AnimNextEditor/Private/Param/ObjectPropertyLocatorEditor.cpp
--- a/Engine/Plugins/Experimental/AnimNext/Source/AnimNextEditor/Private/Param/ObjectPropertyLocatorEditor.cpp
+++ b/Engine/Plugins/Experimental/AnimNext/Source/AnimNextEditor/Private/Param/ObjectPropertyLocatorEditor.cpp
@@ -194,7 +194,7 @@ UE::UniversalObjectLocator::ELocatorFragmentEditorType FObjectPropertyLocatorEdi
 
 bool FObjectPropertyLocatorEditor::IsAllowedInContext(FName InContextName) const
 {
-	return InContextName == "AnimNextContext";
+	return InContextName == "UAFContext";
 }
 
 bool FObjectPropertyLocatorEditor::IsDragSupported(TSharedPtr<FDragDropOperation> DragOperation, UObject* Context) const
diff --git a/Engine/Plugins/Experimental/AnimNext/Source/AnimNextEditor/Private/Variables/VariableOverrideCommands.cpp b/Engine/Plugins/Experimental/AnimNext/Source/AnimNextEditor/Private/Variables/VariableOverrideCommands.cpp
--- a/Engine/Plugins/Experimental/AnimNext/Source/AnimNextEditor/Private/Variables/VariableOverrideCommands.cpp
+++ b/Engine/Plugins/Experimental/AnimNext/Source/AnimNextEditor/Private/Variables/VariableOverrideCommands.cpp
@@ -1,4 +1,4 @@
-﻿// Copyright Epic Games, Inc. All Rights Reserved.
+// Copyright Epic Games, Inc. All Rights Reserved.
 
 #include "VariableOverrideCommands.h"
 
@@ -8,7 +8,7 @@ namespace UE::AnimNext
 {
 
 FVariableOverrideCommands::FVariableOverrideCommands()
-	: TCommands<FVariableOverrideCommands>("AnimNext", LOCTEXT("AnimNextVariableOverrides", "Variable Overrides"), NAME_None, "AnimNextStyle")
+	: TCommands<FVariableOverrideCommands>("UAF", LOCTEXT("UAFVariableOverrides", "Variable Overrides"), NAME_None, "AnimNextStyle")
 {
 }
 
diff --git a/Engine/Plugins/Experimental/AnimNext/Source/AnimNextEditor/Private/Workspace/AnimNextWorkspaceFactory.cpp b/Engine/Plugins/Experimental/AnimNext/Source/AnimNextEditor/Private/Workspace/AnimNextWorkspaceFactory.cpp
--- a/Engine/Plugins/Experimental/AnimNext/Source/AnimNextEditor/Private/Workspace/AnimNextWorkspaceFactory.cpp
+++ b/Engine/Plugins/Experimental/AnimNext/Source/AnimNextEditor/Private/Workspace/AnimNextWorkspaceFactory.cpp
@@ -1,4 +1,4 @@
-﻿// Copyright Epic Games, Inc. All Rights Reserved.
+// Copyright Epic Games, Inc. All Rights Reserved.
 
 #include "Workspace/AnimNextWorkspaceFactory.h"
 #include "AnimNextWorkspaceSchema.h"
@@ -18,7 +18,7 @@ bool UAnimNextWorkspaceFactory::ConfigureProperties()
 
 FText UAnimNextWorkspaceFactory::GetDisplayName() const
 {
-	return LOCTEXT("DisplayName", "AnimNext Workspace");
+	return LOCTEXT("DisplayName", "UAF Workspace");
 }
 
 uint32 UAnimNextWorkspaceFactory::GetMenuCategories() const
@@ -28,7 +28,7 @@ uint32 UAnimNextWorkspaceFactory::GetMenuCategories() const
 
 const TArray<FText>& UAnimNextWorkspaceFactory::GetMenuCategorySubMenus() const
 {
-	static const TArray<FText> Categories = { LOCTEXT("AnimNextSubMenu", "AnimNext") };
+	static const TArray<FText> Categories = { LOCTEXT("UAFSubMenu", "Animation Framework") };
 	return Categories;
 }
 
diff --git a/Engine/Plugins/Experimental/AnimNext/Source/AnimNextEditor/Private/Workspace/AnimNextWorkspaceSchema.cpp b/Engine/Plugins/Experimental/AnimNext/Source/AnimNextEditor/Private/Workspace/AnimNextWorkspaceSchema.cpp
--- a/Engine/Plugins/Experimental/AnimNext/Source/AnimNextEditor/Private/Workspace/AnimNextWorkspaceSchema.cpp
+++ b/Engine/Plugins/Experimental/AnimNext/Source/AnimNextEditor/Private/Workspace/AnimNextWorkspaceSchema.cpp
@@ -1,4 +1,4 @@
-﻿// Copyright Epic Games, Inc. All Rights Reserved.
+// Copyright Epic Games, Inc. All Rights Reserved.
 
 #include "AnimNextWorkspaceSchema.h"
 
@@ -12,7 +12,7 @@
 
 FText UAnimNextWorkspaceSchema::GetDisplayName() const
 {
-	return LOCTEXT("DisplayName", "AnimNext Workspace");
+	return LOCTEXT("DisplayName", "UAF Workspace");
 }
 
 TConstArrayView<FTopLevelAssetPath> UAnimNextWorkspaceSchema::GetSupportedAssetClassPaths() const
diff --git a/Engine/Plugins/Experimental/AnimNext/Source/AnimNextUncookedOnly/Internal/AnimNextRigVMAssetEditorData.h b/Engine/Plugins/Experimental/AnimNext/Source/AnimNextUncookedOnly/Internal/AnimNextRigVMAssetEditorData.h
--- a/Engine/Plugins/Experimental/AnimNext/Source/AnimNextUncookedOnly/Internal/AnimNextRigVMAssetEditorData.h
+++ b/Engine/Plugins/Experimental/AnimNext/Source/AnimNextUncookedOnly/Internal/AnimNextRigVMAssetEditorData.h
@@ -98,23 +98,23 @@ class UAnimNextRigVMAssetLibrary : public UBlueprintFunctionLibrary
 
 public:
 	/** Finds an entry in an AnimNext asset */
-	UFUNCTION(BlueprintCallable, Category = "AnimNext|Entries", meta=(ScriptMethod))
+	UFUNCTION(BlueprintCallable, Category = "UAF|Entries", meta=(ScriptMethod))
 	static ANIMNEXTUNCOOKEDONLY_API UAnimNextRigVMAssetEntry* FindEntry(UAnimNextRigVMAsset* InAsset, FName InName);
 
 	/** Removes an entry from an AnimNext asset */
-	UFUNCTION(BlueprintCallable, Category = "AnimNext|Entries", meta=(ScriptMethod))
+	UFUNCTION(BlueprintCallable, Category = "UAF|Entries", meta=(ScriptMethod))
 	static ANIMNEXTUNCOOKEDONLY_API bool RemoveEntry(UAnimNextRigVMAsset* InAsset, UAnimNextRigVMAssetEntry* InEntry, bool bSetupUndoRedo = true, bool bPrintPythonCommand = true);
 
 	/** Removes multiple entries from an AnimNext asset */
-	UFUNCTION(BlueprintCallable, Category = "AnimNext|Entries", meta=(ScriptMethod))
+	UFUNCTION(BlueprintCallable, Category = "UAF|Entries", meta=(ScriptMethod))
 	static ANIMNEXTUNCOOKEDONLY_API bool RemoveEntries(UAnimNextRigVMAsset* InAsset, const TArray<UAnimNextRigVMAssetEntry*>& InEntries, bool bSetupUndoRedo = true, bool bPrintPythonCommand = true);
 
 	/** Removes all entries from an AnimNext asset */
-	UFUNCTION(BlueprintCallable, Category = "AnimNext|Entries", meta=(ScriptMethod))
+	UFUNCTION(BlueprintCallable, Category = "UAF|Entries", meta=(ScriptMethod))
 	static ANIMNEXTUNCOOKEDONLY_API bool RemoveAllEntries(UAnimNextRigVMAsset* InAsset, bool bSetupUndoRedo = true, bool bPrintPythonCommand = true);
 
 	/** Adds a parameter to an AnimNext asset */
-	UFUNCTION(BlueprintCallable, Category = "AnimNext|Entries", meta=(ScriptMethod))
+	UFUNCTION(BlueprintCallable, Category = "UAF|Entries", meta=(ScriptMethod))
 	static ANIMNEXTUNCOOKEDONLY_API UAnimNextVariableEntry* AddVariable(UAnimNextRigVMAsset* InAsset, FName InName, EPropertyBagPropertyType InValueType, EPropertyBagContainerType InContainerType = EPropertyBagContainerType::None, const UObject* InValueTypeObject = nullptr, const FString& InDefaultValue = TEXT(""), bool bSetupUndoRedo = true, bool bPrintPythonCommand = true);
 
 	/** Adds an event graph to an AnimNext asset */
@@ -122,11 +122,11 @@ public:
 	static ANIMNEXTUNCOOKEDONLY_API UAnimNextEventGraphEntry* AddEventGraph(UAnimNextRigVMAsset* InAsset, FName InName, UScriptStruct* InEventStruct, bool bSetupUndoRedo = true, bool bPrintPythonCommand = true);
 
 	/** Adds a data interface to an AnimNext asset */
-	UFUNCTION(BlueprintCallable, Category = "AnimNext|Entries", meta=(ScriptMethod))
+	UFUNCTION(BlueprintCallable, Category = "UAF|Entries", meta=(ScriptMethod))
 	static ANIMNEXTUNCOOKEDONLY_API UAnimNextDataInterfaceEntry* AddDataInterface(UAnimNextRigVMAsset* InAsset, UAnimNextDataInterface* InDataInterface, bool bSetupUndoRedo = true, bool bPrintPythonCommand = true);
 
 	/** Adds a function to an AnimNext asset */
-	UFUNCTION(BlueprintCallable, Category = "AnimNext|Entries", meta=(ScriptMethod))
+	UFUNCTION(BlueprintCallable, Category = "UAF|Entries", meta=(ScriptMethod))
 	static ANIMNEXTUNCOOKEDONLY_API URigVMLibraryNode* AddFunction(UAnimNextRigVMAsset* InAsset, FName InFunctionName, bool bInMutable, bool bSetupUndoRedo = true, bool bPrintPythonCommand = true);
 };
 
diff --git a/Engine/Plugins/Experimental/AnimNext/Source/AnimNextUncookedOnly/Internal/Variables/AnimNextUniversalObjectLocatorBindingData.h b/Engine/Plugins/Experimental/AnimNext/Source/AnimNextUncookedOnly/Internal/Variables/AnimNextUniversalObjectLocatorBindingData.h
--- a/Engine/Plugins/Experimental/AnimNext/Source/AnimNextUncookedOnly/Internal/Variables/AnimNextUniversalObjectLocatorBindingData.h
+++ b/Engine/Plugins/Experimental/AnimNext/Source/AnimNextUncookedOnly/Internal/Variables/AnimNextUniversalObjectLocatorBindingData.h
@@ -37,7 +37,7 @@ struct FAnimNextUniversalObjectLocatorBindingData : public FAnimNextVariableBind
 	TSoftObjectPtr<UFunction> Function;
 
 	// Object locator
-	UPROPERTY(EditAnywhere, Category = "Binding", meta = (LocatorContext="AnimNextContext"))
+	UPROPERTY(EditAnywhere, Category = "Binding", meta = (LocatorContext="UAFContext"))
 	FUniversalObjectLocator Locator;
 
 	UPROPERTY(EditAnywhere, Category = "Binding")
diff --git a/Engine/Plugins/Experimental/AnimNextAnimGraph/AnimNextAnimGraph.uplugin b/Engine/Plugins/Experimental/AnimNextAnimGraph/AnimNextAnimGraph.uplugin
--- a/Engine/Plugins/Experimental/AnimNextAnimGraph/AnimNextAnimGraph.uplugin
+++ b/Engine/Plugins/Experimental/AnimNextAnimGraph/AnimNextAnimGraph.uplugin
@@ -2,7 +2,7 @@
 	"FileVersion" : 3,
 	"Version" : 1,
 	"VersionName" : "0.1",
-	"FriendlyName" : "Anim Next Anim Graph",
+	"FriendlyName" : "UAF Anim Graph",
 	"Description" : "Framework for defining animation graphs.",
 	"Category" : "Animation",
 	"CreatedBy" : "Epic Games, Inc.",
diff --git a/Engine/Plugins/Experimental/AnimNextAnimGraph/Source/AnimNextAnimGraph/Private/AnimNextAnimGraphModule.cpp b/Engine/Plugins/Experimental/AnimNextAnimGraph/Source/AnimNextAnimGraph/Private/AnimNextAnimGraphModule.cpp
--- a/Engine/Plugins/Experimental/AnimNextAnimGraph/Source/AnimNextAnimGraph/Private/AnimNextAnimGraphModule.cpp
+++ b/Engine/Plugins/Experimental/AnimNextAnimGraph/Source/AnimNextAnimGraph/Private/AnimNextAnimGraphModule.cpp
@@ -88,13 +88,13 @@ void FAnimNextAnimGraphModule::StartupModule()
 	if (!IsRunningCommandlet())
 	{
 		ConsoleCommands.Add(IConsoleManager::Get().RegisterConsoleCommand(
-			TEXT("AnimNext.ListNodeTemplates"),
+			TEXT("UAF.ListNodeTemplates"),
 			TEXT("Dumps statistics about node templates to the log."),
 			FConsoleCommandWithArgsDelegate::CreateRaw(this, &FAnimNextAnimGraphModule::ListNodeTemplates),
 			ECVF_Default
 		));
 		ConsoleCommands.Add(IConsoleManager::Get().RegisterConsoleCommand(
-			TEXT("AnimNext.Module"),
+			TEXT("UAF.Module"),
 			TEXT("Dumps statistics about modules to the log."),
 			FConsoleCommandWithArgsDelegate::CreateRaw(this, &FAnimNextAnimGraphModule::ListAnimationGraphs),
 			ECVF_Default
diff --git a/Engine/Plugins/Experimental/AnimNextAnimGraph/Source/AnimNextAnimGraph/Public/AnimNextAnimGraphStats.h b/Engine/Plugins/Experimental/AnimNextAnimGraph/Source/AnimNextAnimGraph/Public/AnimNextAnimGraphStats.h
--- a/Engine/Plugins/Experimental/AnimNextAnimGraph/Source/AnimNextAnimGraph/Public/AnimNextAnimGraphStats.h
+++ b/Engine/Plugins/Experimental/AnimNextAnimGraph/Source/AnimNextAnimGraph/Public/AnimNextAnimGraphStats.h
@@ -5,8 +5,8 @@
 #include "Stats/Stats.h"
 #include "AnimNextStats.h"
 
-DECLARE_CYCLE_STAT_EXTERN(TEXT("AnimNext: Allocate Graph Instance"), STAT_AnimNext_Graph_AllocateInstance, STATGROUP_AnimNext, ANIMNEXTANIMGRAPH_API);
-DECLARE_CYCLE_STAT_EXTERN(TEXT("AnimNext: Graph RigVM"), STAT_AnimNext_Graph_RigVM, STATGROUP_AnimNext, ANIMNEXTANIMGRAPH_API);
-DECLARE_CYCLE_STAT_EXTERN(TEXT("AnimNext: Execute Evaluation Program"), STAT_AnimNext_EvaluationProgram_Execute, STATGROUP_AnimNext, ANIMNEXTANIMGRAPH_API);
+DECLARE_CYCLE_STAT_EXTERN(TEXT("UAF: Allocate Graph Instance"), STAT_AnimNext_Graph_AllocateInstance, STATGROUP_AnimNext, ANIMNEXTANIMGRAPH_API);
+DECLARE_CYCLE_STAT_EXTERN(TEXT("UAF: Graph RigVM"), STAT_AnimNext_Graph_RigVM, STATGROUP_AnimNext, ANIMNEXTANIMGRAPH_API);
+DECLARE_CYCLE_STAT_EXTERN(TEXT("UAF: Execute Evaluation Program"), STAT_AnimNext_EvaluationProgram_Execute, STATGROUP_AnimNext, ANIMNEXTANIMGRAPH_API);
 
-DECLARE_CYCLE_STAT_EXTERN(TEXT("AnimNext: Call Function"), STAT_AnimNext_CallFunction, STATGROUP_AnimNext, ANIMNEXTANIMGRAPH_API);
+DECLARE_CYCLE_STAT_EXTERN(TEXT("UAF: Call Function"), STAT_AnimNext_CallFunction, STATGROUP_AnimNext, ANIMNEXTANIMGRAPH_API);
diff --git a/Engine/Plugins/Experimental/AnimNextAnimGraph/Source/AnimNextAnimGraph/Public/EvaluationVM/Tasks/PushPose.h b/Engine/Plugins/Experimental/AnimNextAnimGraph/Source/AnimNextAnimGraph/Public/EvaluationVM/Tasks/PushPose.h
--- a/Engine/Plugins/Experimental/AnimNextAnimGraph/Source/AnimNextAnimGraph/Public/EvaluationVM/Tasks/PushPose.h
+++ b/Engine/Plugins/Experimental/AnimNextAnimGraph/Source/AnimNextAnimGraph/Public/EvaluationVM/Tasks/PushPose.h
@@ -11,7 +11,7 @@
 
 #define UE_API ANIMNEXTANIMGRAPH_API
 
-DECLARE_CYCLE_STAT_EXTERN(TEXT("AnimNext Task: PushPose"), STAT_AnimNext_Task_PushPose, STATGROUP_AnimNext, ANIMNEXTANIMGRAPH_API);
+DECLARE_CYCLE_STAT_EXTERN(TEXT("UAF Task: PushPose"), STAT_AnimNext_Task_PushPose, STATGROUP_AnimNext, ANIMNEXTANIMGRAPH_API);
 
 USTRUCT()
 struct FAnimNextPushPoseTask : public FAnimNextEvaluationTask
diff --git a/Engine/Plugins/Experimental/AnimNextAnimGraph/Source/AnimNextAnimGraph/Public/Graph/AnimNextAnimGraph.h b/Engine/Plugins/Experimental/AnimNextAnimGraph/Source/AnimNextAnimGraph/Public/Graph/AnimNextAnimGraph.h
--- a/Engine/Plugins/Experimental/AnimNextAnimGraph/Source/AnimNextAnimGraph/Public/Graph/AnimNextAnimGraph.h
+++ b/Engine/Plugins/Experimental/AnimNextAnimGraph/Source/AnimNextAnimGraph/Public/Graph/AnimNextAnimGraph.h
@@ -65,7 +65,7 @@ private:
 //~ This struct uses UE reflection because we wish for the GC to keep the graph
 //~ alive while we own a reference to it (specifically when running an animation graph).
 //~ It is not intended to be serialized on disk with a live instance.
-USTRUCT(BlueprintType, DisplayName="AnimNext Graph")
+USTRUCT(BlueprintType, DisplayName="UAF Graph")
 struct FAnimNextAnimGraph
 {
 	GENERATED_BODY()
diff --git a/Engine/Plugins/Experimental/AnimNextAnimGraph/Source/AnimNextAnimGraph/Public/GraphInterfaces/AnimNextNativeDataInterface_AnimSequencePlayer.h b/Engine/Plugins/Experimental/AnimNextAnimGraph/Source/AnimNextAnimGraph/Public/GraphInterfaces/AnimNextNativeDataInterface_AnimSequencePlayer.h
--- a/Engine/Plugins/Experimental/AnimNextAnimGraph/Source/AnimNextAnimGraph/Public/GraphInterfaces/AnimNextNativeDataInterface_AnimSequencePlayer.h
+++ b/Engine/Plugins/Experimental/AnimNextAnimGraph/Source/AnimNextAnimGraph/Public/GraphInterfaces/AnimNextNativeDataInterface_AnimSequencePlayer.h
@@ -1,4 +1,4 @@
-﻿// Copyright Epic Games, Inc. All Rights Reserved.
+// Copyright Epic Games, Inc. All Rights Reserved.
 
 #pragma once
 
@@ -11,7 +11,7 @@ class UAnimSequence;
 /**
 * Native interface to built-in AnimSequence player graph at: /AnimNextAnimGraph/FactoryGraphs/AG_AnimSequencePlayer
  */
-USTRUCT(BlueprintType, DisplayName = "AnimNext Anim Sequence Player Data")
+USTRUCT(BlueprintType, DisplayName = "UAF Anim Sequence Player Data")
 struct FAnimNextNativeDataInterface_AnimSequencePlayer : public FAnimNextNativeDataInterface
 {
 	GENERATED_BODY()
diff --git a/Engine/Plugins/Experimental/AnimNextAnimGraph/Source/AnimNextAnimGraph/Public/GraphInterfaces/AnimNextNativeDataInterface_BlendSpacePlayer.h b/Engine/Plugins/Experimental/AnimNextAnimGraph/Source/AnimNextAnimGraph/Public/GraphInterfaces/AnimNextNativeDataInterface_BlendSpacePlayer.h
--- a/Engine/Plugins/Experimental/AnimNextAnimGraph/Source/AnimNextAnimGraph/Public/GraphInterfaces/AnimNextNativeDataInterface_BlendSpacePlayer.h
+++ b/Engine/Plugins/Experimental/AnimNextAnimGraph/Source/AnimNextAnimGraph/Public/GraphInterfaces/AnimNextNativeDataInterface_BlendSpacePlayer.h
@@ -11,7 +11,7 @@ class UBlendSpace;
 /**
 * Native interface to built-in BlendSpace player graph at: /AnimNextAnimGraph/FactoryGraphs/AG_BlendSpacePlayer
  */
-USTRUCT(BlueprintType, DisplayName = "AnimNext Blend Space Player Data")
+USTRUCT(BlueprintType, DisplayName = "UAF Blend Space Player Data")
 struct FAnimNextNativeDataInterface_BlendSpacePlayer : public FAnimNextNativeDataInterface
 {
 	GENERATED_BODY()
diff --git a/Engine/Plugins/Experimental/AnimNextAnimGraph/Source/AnimNextAnimGraph/Public/GraphInterfaces/AnimNextNativeDataInterface_SynchronizeUsingGroups.h b/Engine/Plugins/Experimental/AnimNextAnimGraph/Source/AnimNextAnimGraph/Public/GraphInterfaces/AnimNextNativeDataInterface_SynchronizeUsingGroups.h
--- a/Engine/Plugins/Experimental/AnimNextAnimGraph/Source/AnimNextAnimGraph/Public/GraphInterfaces/AnimNextNativeDataInterface_SynchronizeUsingGroups.h
+++ b/Engine/Plugins/Experimental/AnimNextAnimGraph/Source/AnimNextAnimGraph/Public/GraphInterfaces/AnimNextNativeDataInterface_SynchronizeUsingGroups.h
@@ -12,7 +12,7 @@
 /**
 * Native interface to built-in SynchronizeUsingGroups trait
  */
-USTRUCT(BlueprintType, DisplayName = "AnimNext Group Synchronization Data")
+USTRUCT(BlueprintType, DisplayName = "UAF Group Synchronization Data")
 struct FAnimNextNativeDataInterface_SynchronizeUsingGroups : public FAnimNextNativeDataInterface
 {
 	GENERATED_BODY()
diff --git a/Engine/Plugins/Experimental/AnimNextAnimGraph/Source/AnimNextAnimGraphEditor/Private/AnimNextAnimGraphEditorModule.cpp b/Engine/Plugins/Experimental/AnimNextAnimGraph/Source/AnimNextAnimGraphEditor/Private/AnimNextAnimGraphEditorModule.cpp
--- a/Engine/Plugins/Experimental/AnimNextAnimGraph/Source/AnimNextAnimGraphEditor/Private/AnimNextAnimGraphEditorModule.cpp
+++ b/Engine/Plugins/Experimental/AnimNextAnimGraph/Source/AnimNextAnimGraphEditor/Private/AnimNextAnimGraphEditorModule.cpp
@@ -61,8 +61,8 @@ class FAnimNextAnimGraphEditorModule : public IModuleInterface
 		// Register settings
 		ISettingsModule& SettingsModule = FModuleManager::LoadModuleChecked<ISettingsModule>("Settings");
 		SettingsModule.RegisterSettings("Project", "Plugins", "AnimNextAnimGraph",
-			LOCTEXT("SettingsName", "AnimNext Anim Graph"),
-			LOCTEXT("SettingsDescription", "Configure options for AnimNext animation graphs."),
+			LOCTEXT("SettingsName", "UAF Anim Graph"),
+			LOCTEXT("SettingsDescription", "Configure options for UAF animation graphs."),
 			GetMutableDefault<UAnimNextAnimGraphSettings>());
 
 		FPropertyEditorModule& PropertyModule = FModuleManager::LoadModuleChecked<FPropertyEditorModule>("PropertyEditor");
diff --git a/Engine/Plugins/Experimental/AnimNextAnimGraph/Source/AnimNextAnimGraphEditor/Private/Graph/AnimNextAnimationGraphAssetDefinition.h b/Engine/Plugins/Experimental/AnimNextAnimGraph/Source/AnimNextAnimGraphEditor/Private/Graph/AnimNextAnimationGraphAssetDefinition.h
--- a/Engine/Plugins/Experimental/AnimNextAnimGraph/Source/AnimNextAnimGraphEditor/Private/Graph/AnimNextAnimationGraphAssetDefinition.h
+++ b/Engine/Plugins/Experimental/AnimNextAnimGraph/Source/AnimNextAnimGraphEditor/Private/Graph/AnimNextAnimationGraphAssetDefinition.h
@@ -1,4 +1,4 @@
-﻿// Copyright Epic Games, Inc. All Rights Reserved.
+// Copyright Epic Games, Inc. All Rights Reserved.
 
 #pragma once
 
@@ -15,13 +15,13 @@ class UAssetDefinition_AnimNextAnimationGraph : public UAssetDefinitionDefault
 
 public:
 	// UAssetDefinition interface
-	virtual FText GetAssetDisplayName() const override { return LOCTEXT("AnimNextAnimationGraph", "AnimNext Animation Graph"); }
+	virtual FText GetAssetDisplayName() const override { return LOCTEXT("UAFAnimationGraph", "UAF Animation Graph"); }
 	virtual FLinearColor GetAssetColor() const override { return FLinearColor(FColor(128,96,48)); }
 	virtual TSoftClassPtr<UObject> GetAssetClass() const override { return UAnimNextAnimationGraph::StaticClass(); }
 	virtual EAssetCommandResult OpenAssets(const FAssetOpenArgs& OpenArgs) const override;
 	virtual TConstArrayView<FAssetCategoryPath> GetAssetCategories() const override
 	{
-		static const auto Categories = { FAssetCategoryPath(EAssetCategoryPaths::Animation, LOCTEXT("AnimNextSubMenu", "AnimNext")) };
+		static const auto Categories = { FAssetCategoryPath(EAssetCategoryPaths::Animation, LOCTEXT("UAFSubMenu", "Animation Framework")) };
 		return Categories;
 	}
 	virtual bool ShouldSaveExternalPackages() const override { return true; }
diff --git a/Engine/Plugins/Experimental/AnimNextAnimGraph/Source/AnimNextAnimGraphUncookedOnly/Private/Graph/AnimGraph/AnimGraphNode_AnimNextGraph.cpp b/Engine/Plugins/Experimental/AnimNextAnimGraph/Source/AnimNextAnimGraphUncookedOnly/Private/Graph/AnimGraph/AnimGraphNode_AnimNextGraph.cpp
--- a/Engine/Plugins/Experimental/AnimNextAnimGraph/Source/AnimNextAnimGraphUncookedOnly/Private/Graph/AnimGraph/AnimGraphNode_AnimNextGraph.cpp
+++ b/Engine/Plugins/Experimental/AnimNextAnimGraph/Source/AnimNextAnimGraphUncookedOnly/Private/Graph/AnimGraph/AnimGraphNode_AnimNextGraph.cpp
@@ -21,17 +21,17 @@ UAnimGraphNode_AnimNextGraph::UAnimGraphNode_AnimNextGraph(const FObjectInitiali
 
 FText UAnimGraphNode_AnimNextGraph::GetNodeTitle(ENodeTitleType::Type TitleType) const
 {
-	return LOCTEXT("AnimGraphNode_AnimNextModule_Title", "AnimNext Animation Graph");
+	return LOCTEXT("AnimGraphNode_UAFModule_Title", "UAF Animation Graph");
 }
 
 FText UAnimGraphNode_AnimNextGraph::GetTooltipText() const
 {
-	return LOCTEXT("AnimGraphNode_AnimNextAnimationGraph_Tooltip", "Evaluates a AnimNext Animation Graph");
+	return LOCTEXT("AnimGraphNode_UAFAnimationGraph_Tooltip", "Evaluates a UAF Animation Graph");
 }
 
 FText UAnimGraphNode_AnimNextGraph::GetMenuCategory() const
 {
-	return LOCTEXT("Category", "AnimNext");
+	return LOCTEXT("Category", "Animation Framework");
 }
 
 void UAnimGraphNode_AnimNextGraph::PreloadRequiredAssets()
diff --git a/Engine/Plugins/Experimental/AnimNextChooser/AnimNextChooser.uplugin b/Engine/Plugins/Experimental/AnimNextChooser/AnimNextChooser.uplugin
--- a/Engine/Plugins/Experimental/AnimNextChooser/AnimNextChooser.uplugin
+++ b/Engine/Plugins/Experimental/AnimNextChooser/AnimNextChooser.uplugin
@@ -2,8 +2,8 @@
 	"FileVersion" : 3,
 	"Version" : 1,
 	"VersionName" : "0.1",
-	"FriendlyName" : "Anim Next Chooser",
-	"Description" : "Chooser integration for Anim Next.",
+	"FriendlyName" : "UAF Chooser",
+	"Description" : "Chooser integration for UAF.",
 	"Category" : "Animation",
 	"CreatedBy" : "Epic Games, Inc.",
 	"CreatedByURL" : "https://epicgames.com",
diff --git a/Engine/Plugins/Experimental/AnimNextControlRig/AnimNextControlRig.uplugin b/Engine/Plugins/Experimental/AnimNextControlRig/AnimNextControlRig.uplugin
--- a/Engine/Plugins/Experimental/AnimNextControlRig/AnimNextControlRig.uplugin
+++ b/Engine/Plugins/Experimental/AnimNextControlRig/AnimNextControlRig.uplugin
@@ -2,8 +2,8 @@
 	"FileVersion" : 3,
 	"Version" : 1,
 	"VersionName" : "0.1",
-	"FriendlyName" : "Anim Next Control Rig",
-	"Description" : "Control Rig integration for Anim Next.",
+	"FriendlyName" : "UAF Control Rig",
+	"Description" : "Control Rig integration for UAF.",
 	"Category" : "Animation",
 	"CreatedBy" : "Epic Games, Inc.",
 	"CreatedByURL" : "https://epicgames.com",
diff --git a/Engine/Plugins/Experimental/AnimNextControlRig/Source/AnimNextControlRig/Private/ControlRigTask.h b/Engine/Plugins/Experimental/AnimNextControlRig/Source/AnimNextControlRig/Private/ControlRigTask.h
--- a/Engine/Plugins/Experimental/AnimNextControlRig/Source/AnimNextControlRig/Private/ControlRigTask.h
+++ b/Engine/Plugins/Experimental/AnimNextControlRig/Source/AnimNextControlRig/Private/ControlRigTask.h
@@ -16,7 +16,7 @@ namespace UE::AnimNext
 struct FKeyframeState;
 }
 
-DECLARE_CYCLE_STAT_EXTERN(TEXT("AnimNext Task: ControlRig"), STAT_AnimNext_Task_ControlRig, STATGROUP_AnimNext, ANIMNEXTCONTROLRIG_API);
+DECLARE_CYCLE_STAT_EXTERN(TEXT("UAF Task: ControlRig"), STAT_AnimNext_Task_ControlRig, STATGROUP_AnimNext, ANIMNEXTCONTROLRIG_API);
 
 USTRUCT()
 struct ANIMNEXTCONTROLRIG_API FAnimNextControlRigTask : public FAnimNextEvaluationTask
diff --git a/Engine/Plugins/Experimental/AnimNextPoseSearch/AnimNextPoseSearch.uplugin b/Engine/Plugins/Experimental/AnimNextPoseSearch/AnimNextPoseSearch.uplugin
--- a/Engine/Plugins/Experimental/AnimNextPoseSearch/AnimNextPoseSearch.uplugin
+++ b/Engine/Plugins/Experimental/AnimNextPoseSearch/AnimNextPoseSearch.uplugin
@@ -2,8 +2,8 @@
 	"FileVersion" : 3,
 	"Version" : 1,
 	"VersionName" : "0.1",
-	"FriendlyName" : "Anim Next Pose Search",
-	"Description" : "Pose Search integration for Anim Next.",
+	"FriendlyName" : "UAF Pose Search",
+	"Description" : "Pose Search integration for UAF.",
 	"Category" : "Animation",
 	"CreatedBy" : "Epic Games, Inc.",
 	"CreatedByURL" : "https://epicgames.com",
diff --git a/Engine/Plugins/Experimental/AnimNextStateTree/AnimNextStateTree.uplugin b/Engine/Plugins/Experimental/AnimNextStateTree/AnimNextStateTree.uplugin
--- a/Engine/Plugins/Experimental/AnimNextStateTree/AnimNextStateTree.uplugin
+++ b/Engine/Plugins/Experimental/AnimNextStateTree/AnimNextStateTree.uplugin
@@ -2,8 +2,8 @@
 	"FileVersion" : 3,
 	"Version" : 1,
 	"VersionName" : "0.1",
-	"FriendlyName" : "Anim Next State Tree",
-	"Description" : "StateTree integration for AnimNext.",
+	"FriendlyName" : "UAF State Tree",
+	"Description" : "StateTree integration for UAF.",
 	"Category" : "Animation",
 	"CreatedBy" : "Epic Games, Inc.",
 	"CreatedByURL" : "https://epicgames.com",
diff --git a/Engine/Plugins/Experimental/AnimNextStateTree/Source/AnimNextStateTree/Public/AnimNextStateTreeTypes.h b/Engine/Plugins/Experimental/AnimNextStateTree/Source/AnimNextStateTree/Public/AnimNextStateTreeTypes.h
--- a/Engine/Plugins/Experimental/AnimNextStateTree/Source/AnimNextStateTree/Public/AnimNextStateTreeTypes.h
+++ b/Engine/Plugins/Experimental/AnimNextStateTree/Source/AnimNextStateTree/Public/AnimNextStateTreeTypes.h
@@ -11,7 +11,7 @@
 /**
  * Base struct for all AnimNext StateTree Evaluators.
  */
-USTRUCT(meta = (DisplayName = "AnimNext Evaluator Base", Hidden))
+USTRUCT(meta = (DisplayName = "UAF Evaluator Base", Hidden))
 struct ANIMNEXTSTATETREE_API FAnimNextStateTreeEvaluatorBase : public FStateTreeEvaluatorBase
 {
 	GENERATED_BODY()
@@ -20,7 +20,7 @@ struct ANIMNEXTSTATETREE_API FAnimNextStateTreeEvaluatorBase : public FStateTree
 /**
  * Base struct for all AnimNext StateTree Tasks.
  */
-USTRUCT(meta = (DisplayName = "AnimNext Task Base", Hidden))
+USTRUCT(meta = (DisplayName = "UAF Task Base", Hidden))
 struct ANIMNEXTSTATETREE_API FAnimNextStateTreeTaskBase : public FStateTreeTaskBase
 {
 	GENERATED_BODY()
diff --git a/Engine/Plugins/Experimental/AnimNextStateTree/Source/AnimNextStateTreeEditor/Private/AnimNextStateTreeAssetDefinition.h b/Engine/Plugins/Experimental/AnimNextStateTree/Source/AnimNextStateTreeEditor/Private/AnimNextStateTreeAssetDefinition.h
--- a/Engine/Plugins/Experimental/AnimNextStateTree/Source/AnimNextStateTreeEditor/Private/AnimNextStateTreeAssetDefinition.h
+++ b/Engine/Plugins/Experimental/AnimNextStateTree/Source/AnimNextStateTreeEditor/Private/AnimNextStateTreeAssetDefinition.h
@@ -15,13 +15,13 @@ class UAssetDefinition_AnimNextStateTree : public UAssetDefinitionDefault
 
 public:
 	// UAssetDefinition interface
-	virtual FText GetAssetDisplayName() const override { return LOCTEXT("AnimNextStateTree", "AnimNext State Tree"); }
+	virtual FText GetAssetDisplayName() const override { return LOCTEXT("UAFStateTree", "UAF State Tree"); }
 	virtual FLinearColor GetAssetColor() const override { return FLinearColor(FColor(96,128,48)); }
 	virtual TSoftClassPtr<UObject> GetAssetClass() const override { return UAnimNextStateTree::StaticClass(); }
 	virtual EAssetCommandResult OpenAssets(const FAssetOpenArgs& OpenArgs) const override;
 	virtual TConstArrayView<FAssetCategoryPath> GetAssetCategories() const override
 	{
-		static const auto Categories = { FAssetCategoryPath(EAssetCategoryPaths::Animation, LOCTEXT("AnimNextSubMenu", "AnimNext")) };
+		static const auto Categories = { FAssetCategoryPath(EAssetCategoryPaths::Animation, LOCTEXT("UAFSubMenu", "Animation Framework")) };
 		return Categories;
 	}
 	virtual bool ShouldSaveExternalPackages() const override { return true; }
diff --git a/Engine/Plugins/Experimental/AnimNextWarping/AnimNextWarping.uplugin b/Engine/Plugins/Experimental/AnimNextWarping/AnimNextWarping.uplugin
--- a/Engine/Plugins/Experimental/AnimNextWarping/AnimNextWarping.uplugin
+++ b/Engine/Plugins/Experimental/AnimNextWarping/AnimNextWarping.uplugin
@@ -2,8 +2,8 @@
 	"FileVersion" : 3,
 	"Version" : 1,
 	"VersionName" : "0.1",
-	"FriendlyName" : "Anim Next Warping",
-	"Description": "Framework for animation and pose warping for Anim Next.",
+	"FriendlyName" : "UAF Warping",
+	"Description": "Framework for animation and pose warping for UAF.",
 	"Category" : "Animation",
 	"CreatedBy" : "Epic Games, Inc.",
 	"CreatedByURL" : "https://epicgames.com",
diff --git a/Engine/Plugins/Experimental/Animation/SkeletonTemplateFramework/Source/Editor/Private/SkeletonTemplateFramework/SkeletonBindingAssetDefinition.cpp b/Engine/Plugins/Experimental/Animation/SkeletonTemplateFramework/Source/Editor/Private/SkeletonTemplateFramework/SkeletonBindingAssetDefinition.cpp
--- a/Engine/Plugins/Experimental/Animation/SkeletonTemplateFramework/Source/Editor/Private/SkeletonTemplateFramework/SkeletonBindingAssetDefinition.cpp
+++ b/Engine/Plugins/Experimental/Animation/SkeletonTemplateFramework/Source/Editor/Private/SkeletonTemplateFramework/SkeletonBindingAssetDefinition.cpp
@@ -34,7 +34,7 @@ EAssetCommandResult UAssetDefinition_SkeletonBinding::OpenAssets(const FAssetOpe
 
 TConstArrayView<FAssetCategoryPath> UAssetDefinition_SkeletonBinding::GetAssetCategories() const
 {
-	static const auto Categories = { FAssetCategoryPath(EAssetCategoryPaths::Animation, LOCTEXT("AnimNextSubMenu", "AnimNext")) };
+	static const auto Categories = { FAssetCategoryPath(EAssetCategoryPaths::Animation, LOCTEXT("UAFSubMenu", "Animation Framework")) };
 	return Categories;
 }
 
diff --git a/Engine/Plugins/Experimental/Animation/SkeletonTemplateFramework/Source/Editor/Private/SkeletonTemplateFramework/SkeletonTemplateAssetDefinition.cpp b/Engine/Plugins/Experimental/Animation/SkeletonTemplateFramework/Source/Editor/Private/SkeletonTemplateFramework/SkeletonTemplateAssetDefinition.cpp
--- a/Engine/Plugins/Experimental/Animation/SkeletonTemplateFramework/Source/Editor/Private/SkeletonTemplateFramework/SkeletonTemplateAssetDefinition.cpp
+++ b/Engine/Plugins/Experimental/Animation/SkeletonTemplateFramework/Source/Editor/Private/SkeletonTemplateFramework/SkeletonTemplateAssetDefinition.cpp
@@ -35,7 +35,7 @@ EAssetCommandResult UAssetDefinition_SkeletonTemplate::OpenAssets(const FAssetOp
 
 TConstArrayView<FAssetCategoryPath> UAssetDefinition_SkeletonTemplate::GetAssetCategories() const
 {
-	static const auto Categories = { FAssetCategoryPath(EAssetCategoryPaths::Animation, LOCTEXT("AnimNextSubMenu", "AnimNext")) };
+	static const auto Categories = { FAssetCategoryPath(EAssetCategoryPaths::Animation, LOCTEXT("UAFSubMenu", "Animation Framework")) };
 	return Categories;
 }
 
