Use unaligned SIMD loads where we don't guarantee 16-byte alignment

diff --git a/Engine/Plugins/Experimental/AnimNext/Source/AnimNext/Private/GenerationTools.cpp b/Engine/Plugins/Experimental/AnimNext/Source/AnimNext/Private/GenerationTools.cpp
--- a/Engine/Plugins/Experimental/AnimNext/Source/AnimNext/Private/GenerationTools.cpp
+++ b/Engine/Plugins/Experimental/AnimNext/Source/AnimNext/Private/GenerationTools.cpp
@@ -683,8 +683,8 @@ void FGenerationTools::FixLODRequiredBones(const int32 NumLODs
 	while (RotationPtr < RotationEndPtr)
 	{
 		VectorRegister4Double Rotation = VectorLoadAligned(RotationPtr);
-		VectorRegister4Double Translation = VectorSet_W0(VectorLoadAligned(TranslationPtr));
-		VectorRegister4Double Scale = VectorSet_W0(VectorLoadAligned(TranslationPtr + ScaleOffset));
+		VectorRegister4Double Translation = VectorSet_W0(VectorLoad(TranslationPtr));
+		VectorRegister4Double Scale = VectorSet_W0(VectorLoad(TranslationPtr + ScaleOffset));
 
 		const int32 DestTransformIndex = *LODBoneIndexPtr;
 		TargetTransforms[DestTransformIndex] = FTransform(Rotation, Translation, Scale);
diff --git a/Engine/Plugins/Experimental/AnimNext/Source/AnimNext/Private/Graph/RemapPoseData.cpp b/Engine/Plugins/Experimental/AnimNext/Source/AnimNext/Private/Graph/RemapPoseData.cpp
--- a/Engine/Plugins/Experimental/AnimNext/Source/AnimNext/Private/Graph/RemapPoseData.cpp
+++ b/Engine/Plugins/Experimental/AnimNext/Source/AnimNext/Private/Graph/RemapPoseData.cpp
@@ -167,8 +167,8 @@ void FRemapPoseData::RemapPose(const UE::AnimNext::FLODPoseHeap& SourcePose,
 		OutTargetPose.LocalTransformsView.Rotations.RangeCheck(TargetBoneIndex);
 
 		VectorRegister4Double Rotation = VectorLoadAligned(SrcRotationPtr + SourceBoneIndex * 4);
-		VectorRegister4Double Translation = VectorLoadAligned(SrcTranslationPtr + SourceBoneIndex * 3);
-		VectorRegister4Double Scale = VectorLoadAligned(SrcTranslationPtr + SourceBoneIndex * 3 + SrcScaleOffset);
+		VectorRegister4Double Translation = VectorLoad(SrcTranslationPtr + SourceBoneIndex * 3);
+		VectorRegister4Double Scale = VectorLoad(SrcTranslationPtr + SourceBoneIndex * 3 + SrcScaleOffset);
 
 		VectorStoreAligned(Rotation, DstRotationPtr + TargetBoneIndex * 4);
 		VectorStoreFloat3(Translation, DstTranslationPtr + TargetBoneIndex * 3);
diff --git a/Engine/Plugins/Experimental/AnimNext/Source/AnimNext/Private/TransformArrayOperations_SoA.cpp b/Engine/Plugins/Experimental/AnimNext/Source/AnimNext/Private/TransformArrayOperations_SoA.cpp
--- a/Engine/Plugins/Experimental/AnimNext/Source/AnimNext/Private/TransformArrayOperations_SoA.cpp
+++ b/Engine/Plugins/Experimental/AnimNext/Source/AnimNext/Private/TransformArrayOperations_SoA.cpp
@@ -105,12 +105,12 @@ namespace UE::AnimNext
 		while (SrcRotationPtr < SrcRotationLastPtr)
 		{
 			VectorRegister4Double Rotation = VectorLoadAligned(SrcRotationPtr);
-			VectorRegister4Double Translation = VectorLoadAligned(SrcTranslationPtr);
-			VectorRegister4Double Scale = VectorLoadAligned(SrcTranslationPtr + SrcScaleOffset);
+			VectorRegister4Double Translation = VectorLoad(SrcTranslationPtr);
+			VectorRegister4Double Scale = VectorLoad(SrcTranslationPtr + SrcScaleOffset);
 
 			VectorStoreAligned(Rotation, DstRotationPtr);
-			VectorStoreAligned(Translation, DstTranslationPtr);
-			VectorStoreAligned(Scale, DstTranslationPtr + DstScaleOffset);
+			VectorStore(Translation, DstTranslationPtr);
+			VectorStore(Scale, DstTranslationPtr + DstScaleOffset);
 
 			SrcRotationPtr += 4;
 			SrcTranslationPtr += 3;
@@ -122,8 +122,8 @@ namespace UE::AnimNext
 		if (SrcRotationPtr < SrcRotationEndPtr)
 		{
 			VectorRegister4Double Rotation = VectorLoadAligned(SrcRotationPtr);
-			VectorRegister4Double Translation = VectorLoadAligned(SrcTranslationPtr);
-			VectorRegister4Double Scale = VectorLoadAligned(SrcTranslationPtr + SrcScaleOffset);
+			VectorRegister4Double Translation = VectorLoad(SrcTranslationPtr);
+			VectorRegister4Double Scale = VectorLoad(SrcTranslationPtr + SrcScaleOffset);
 
 			VectorStoreAligned(Rotation, DstRotationPtr);
 			VectorStoreFloat3(Translation, DstTranslationPtr);
