Reduce memory churn by avoiding shrinking pose buffers that are re-used frame to frame

diff --git a/Engine/Plugins/Experimental/AnimNext/Source/AnimNext/Public/LODPose.h b/Engine/Plugins/Experimental/AnimNext/Source/AnimNext/Public/LODPose.h
--- a/Engine/Plugins/Experimental/AnimNext/Source/AnimNext/Public/LODPose.h
+++ b/Engine/Plugins/Experimental/AnimNext/Source/AnimNext/Public/LODPose.h
@@ -350,7 +350,7 @@ struct TLODPose : public FLODPose
 		RefPose = &InRefPose;
 
 		const int32 NumTransforms = InRefPose.GetNumBonesForLOD(InLODLevel);
-		LocalTransforms.SetNumUninitialized(NumTransforms);
+		LocalTransforms.SetNumUninitialized(NumTransforms, EAllowShrinking::No);	// Don't shrink, our LOD might increase back causing us to churn
 		LocalTransformsView = LocalTransforms.GetView();
 
 		Flags = (ELODPoseFlags)(bAdditive ? (Flags | ELODPoseFlags::Additive) : Flags & ELODPoseFlags::Additive);
