Move animnext bounds update onto the game thread to fix race condition.

diff --git a/Engine/Plugins/Experimental/AnimNext/Source/AnimNext/Private/Component/SkinnedMeshComponentExtensions.cpp b/Engine/Plugins/Experimental/AnimNext/Source/AnimNext/Private/Component/SkinnedMeshComponentExtensions.cpp
--- a/Engine/Plugins/Experimental/AnimNext/Source/AnimNext/Private/Component/SkinnedMeshComponentExtensions.cpp
+++ b/Engine/Plugins/Experimental/AnimNext/Source/AnimNext/Private/Component/SkinnedMeshComponentExtensions.cpp
@@ -115,11 +115,11 @@ void FSkinnedMeshComponentExtensions::CompleteAndDispatch(
 		InComponent->FlipEditableSpaceBases();
 		InComponent->bHasValidBoneTransform = true;
 
-		InComponent->InvalidateCachedBounds();
-		InComponent->UpdateBounds();
-
 		if (!UE::Anim::Private::GUseDeferredRenderDataUpdate)
 		{
+			InComponent->InvalidateCachedBounds();
+			InComponent->UpdateBounds();
+
 			// Send updated transforms & bounds to the renderer
 			InComponent->SendRenderDynamicData_Concurrent();
 			InComponent->SendRenderTransform_Concurrent();
@@ -161,6 +161,9 @@ void FSkinnedMeshComponentExtensions::CompleteAndDispatch(
 
 				if (UE::Anim::Private::GUseDeferredRenderDataUpdate)
 				{
+					Component->InvalidateCachedBounds();
+					Component->UpdateBounds();
+
 					// Skip components that have a system tag as they will likely be modified later.
 					if (Component->ComponentTags.IsEmpty())
 					{
