AnimNext: Fix Shipping editor compilation

Correctly wrap trace-related systems

diff --git a/Engine/Plugins/Experimental/AnimNext/Source/AnimNextEditor/Private/AnimNextEditorModule.cpp b/Engine/Plugins/Experimental/AnimNext/Source/AnimNextEditor/Private/AnimNextEditorModule.cpp
--- a/Engine/Plugins/Experimental/AnimNext/Source/AnimNextEditor/Private/AnimNextEditorModule.cpp
+++ b/Engine/Plugins/Experimental/AnimNext/Source/AnimNextEditor/Private/AnimNextEditorModule.cpp
@@ -83,7 +83,9 @@
 namespace
 {
 	FRewindDebuggerAnimNext GRewindDebuggerAnimNext;
+#if ANIMNEXT_TRACE_ENABLED
 	UE::AnimNextEditor::FAnimNextModuleTrackCreator GAnimNextModulesTrackCreator;
+#endif
 	FAnimNextTraceModule GAnimNextTraceModule;
 }
 #endif
@@ -107,7 +109,9 @@ void FAnimNextEditorModule::StartupModule()
 
 #if ANIMNEXT_TRACE_ENABLED
 	IModularFeatures::Get().RegisterModularFeature(IRewindDebuggerExtension::ModularFeatureName, &GRewindDebuggerAnimNext);
+#if ANIMNEXT_TRACE_ENABLED
 	IModularFeatures::Get().RegisterModularFeature(RewindDebugger::IRewindDebuggerTrackCreator::ModularFeatureName, &GAnimNextModulesTrackCreator);
+#endif
 	IModularFeatures::Get().RegisterModularFeature(TraceServices::ModuleFeatureName, &GAnimNextTraceModule);
 #endif // ANIMNEXT_TRACE_ENABLED
 
@@ -428,7 +432,9 @@ void FAnimNextEditorModule::ShutdownModule()
 
 #if ANIMNEXT_TRACE_ENABLED
 	IModularFeatures::Get().UnregisterModularFeature(IRewindDebuggerExtension::ModularFeatureName, &GRewindDebuggerAnimNext);
+#if ANIMNEXT_TRACE_ENABLED
 	IModularFeatures::Get().UnregisterModularFeature(RewindDebugger::IRewindDebuggerTrackCreator::ModularFeatureName, &GAnimNextModulesTrackCreator);
+#endif
 	IModularFeatures::Get().UnregisterModularFeature(TraceServices::ModuleFeatureName, &GAnimNextTraceModule);
 #endif // ANIMNEXT_TRACE_ENABLED
 
