[AnimNext] Fix AnimNext dependencies being applied incorrectly when nested in For loops

Trait data was being reused each time around the loop, so only ever applied the last dependency in the loop multiple times. Fix is to copy the trait data rather than reusing the pointer to the trait scope's memory.
Add RigUnit to tell if we are AnimBP driven or not

diff --git a/Engine/Plugins/Experimental/AnimNext/Source/AnimNext/Private/Graph/RigUnit_IsAnimBPDriven.cpp b/Engine/Plugins/Experimental/AnimNext/Source/AnimNext/Private/Graph/RigUnit_IsAnimBPDriven.cpp
new file mode 100644
--- /dev/null
+++ b/Engine/Plugins/Experimental/AnimNext/Source/AnimNext/Private/Graph/RigUnit_IsAnimBPDriven.cpp
@@ -0,0 +1,9 @@
+// Copyright Epic Games, Inc. All Rights Reserved.
+
+#include "RigUnit_IsAnimBPDriven.h"
+#include "Components/SkeletalMeshComponent.h"
+
+FRigUnit_IsAnimBPDriven_Execute()
+{
+	Result = MeshComponent != nullptr && MeshComponent->bEnableAnimation;
+}
diff --git a/Engine/Plugins/Experimental/AnimNext/Source/AnimNext/Private/Graph/RigUnit_IsAnimBPDriven.h b/Engine/Plugins/Experimental/AnimNext/Source/AnimNext/Private/Graph/RigUnit_IsAnimBPDriven.h
new file mode 100644
--- /dev/null
+++ b/Engine/Plugins/Experimental/AnimNext/Source/AnimNext/Private/Graph/RigUnit_IsAnimBPDriven.h
@@ -0,0 +1,24 @@
+// Copyright Epic Games, Inc. All Rights Reserved.
+
+#pragma once
+
+#include "AnimNextExecuteContext.h"
+#include "Graph/RigUnit_AnimNextBase.h"
+
+#include "RigUnit_IsAnimBPDriven.generated.h"
+
+/** Gets whether or not a skeletal mesh component is being driven by an Animation Blueprint */
+USTRUCT(meta = (DisplayName = "Is AnimBP Driven", Category = "Mesh", NodeColor = "0, 1, 1", Keywords = "Mesh"))
+struct FRigUnit_IsAnimBPDriven : public FRigUnit_AnimNextBase
+{
+	GENERATED_BODY()
+
+	RIGVM_METHOD()
+	void Execute();
+
+	UPROPERTY(EditAnywhere, Category = "Graph", meta = (Input))
+	TObjectPtr<USkeletalMeshComponent> MeshComponent;
+
+	UPROPERTY(EditAnywhere, Category = "Graph", meta = (Output))
+	bool Result = false;
+};
diff --git a/Engine/Plugins/Experimental/AnimNext/Source/AnimNext/Private/Module/RigUnit_AnimNextModuleEventDependencies.cpp b/Engine/Plugins/Experimental/AnimNext/Source/AnimNext/Private/Module/RigUnit_AnimNextModuleEventDependencies.cpp
--- a/Engine/Plugins/Experimental/AnimNext/Source/AnimNext/Private/Module/RigUnit_AnimNextModuleEventDependencies.cpp
+++ b/Engine/Plugins/Experimental/AnimNext/Source/AnimNext/Private/Module/RigUnit_AnimNextModuleEventDependencies.cpp
@@ -22,7 +22,7 @@ template<typename PredicateType>
 static void ApplyDependenciesHelper(const FAnimNextExecuteContext& InContext, const PredicateType& InPredicate)
 {
 	FAnimNextModuleInstance& ModuleInstance = InContext.GetContextData<FAnimNextModuleContextData>().GetModuleInstance();
-	TArray<const FRigVMTrait_ModuleEventDependency*> Dependencies;
+	TArray<TInstancedStruct<FRigVMTrait_ModuleEventDependency>> Dependencies;
 	Dependencies.Reserve(InContext.GetTraits().Num());
 	for(const FRigVMTraitScope& TraitScope : InContext.GetTraits())
 	{
@@ -32,7 +32,12 @@ static void ApplyDependenciesHelper(const FAnimNextExecuteContext& InContext, co
 			continue;
 		}
 
-		Dependencies.Add(DependencyTrait);
+		// Copy the state as the memory is re-used by RigVM and we might get called again with a different value for the
+		// same memory location
+		TInstancedStruct<FRigVMTrait_ModuleEventDependency> Dependency;
+		Dependency.InitializeAsScriptStruct(TraitScope.GetScriptStruct(), reinterpret_cast<const uint8*>(DependencyTrait));
+
+		Dependencies.Add(MoveTemp(Dependency));
 	}
 
 	if(Dependencies.Num() == 0)
@@ -44,9 +49,9 @@ static void ApplyDependenciesHelper(const FAnimNextExecuteContext& InContext, co
 	TArray<FTickFunction*> TickFunctions;
 	TickFunctions.Reserve(Dependencies.Num());
 	bool bTickFunctionFound = false; 
-	for(const FRigVMTrait_ModuleEventDependency* Dependency : Dependencies)
+	for(const TInstancedStruct<FRigVMTrait_ModuleEventDependency>& Dependency : Dependencies)
 	{
-		FTickFunction* TickFunction = ModuleInstance.FindTickFunctionByName(Dependency->EventName);
+		FTickFunction* TickFunction = ModuleInstance.FindTickFunctionByName(Dependency.Get<FRigVMTrait_ModuleEventDependency>().EventName);
 		TickFunctions.Add(TickFunction);
 		bTickFunctionFound |= (TickFunction != nullptr);
 	}
@@ -76,8 +81,8 @@ static void ApplyDependenciesHelper(const FAnimNextExecuteContext& InContext, co
 			}
 			
 			FModuleDependencyContext Context(Object, *TickFunction);
-			const FRigVMTrait_ModuleEventDependency* Dependency = Dependencies[DependencyIndex];
-			InPredicate(Context, *Dependency);
+			const FRigVMTrait_ModuleEventDependency& Dependency = Dependencies[DependencyIndex].Get<FRigVMTrait_ModuleEventDependency>();
+			InPredicate(Context, Dependency);
 		}
 	});
 }
