Clean up data registry over time post-GC

diff --git a/Engine/Plugins/Experimental/AnimNext/Source/AnimNext/Private/DataRegistry.cpp b/Engine/Plugins/Experimental/AnimNext/Source/AnimNext/Private/DataRegistry.cpp
--- a/Engine/Plugins/Experimental/AnimNext/Source/AnimNext/Private/DataRegistry.cpp
+++ b/Engine/Plugins/Experimental/AnimNext/Source/AnimNext/Private/DataRegistry.cpp
@@ -18,6 +18,8 @@ UE::AnimNext::FDataRegistry* GAnimationDataRegistry = nullptr;
 constexpr int32 BASIC_TYPE_ALLOC_BLOCK = 1000;
 FDelegateHandle PostGarbageCollectHandle;
 
+static TAutoConsoleVariable<float> CVar_GCCleanupTimeBudget(TEXT("a.AnimNext.GCCleanupTimeBudget"), 100.0f / (1000.0f * 1000.0f), TEXT("Amount of time in seconds that we can use when reclaiming memory post-GC"));
+
 } // end namespace
 
 namespace UE::AnimNext
@@ -39,6 +41,13 @@ namespace UE::AnimNext
 	{
 		FCoreUObjectDelegates::GetPostGarbageCollect().Remove(PostGarbageCollectHandle);
 
+		if (!GAnimationDataRegistry->GCKeysToCleanup.IsEmpty())
+		{
+			GAnimationDataRegistry->RunGCCleanup(INDEX_NONE);
+
+			FTSTicker::GetCoreTicker().RemoveTicker(GAnimationDataRegistry->GCCleanupTickerHandle);
+		}
+
 		// When the module unloads, we cannot use ScriptStructs and our ParamTypes to free the memory
 		// we allocated. To avoid crashing on an invalid access, we clear the type map which means that
 		// destructors won't run for allocated data. If they were to free memory, it will leak on app exit.
@@ -66,20 +75,90 @@ FDataRegistry* FDataRegistry::Get()
 
 /*static*/ void FDataRegistry::HandlePostGarbageCollect()
 {
+	SCOPED_NAMED_EVENT(AnimNext_DataRegistry_PostGC, FColor::Blue);
+
 	// Compact the registry on GC
 	if (GAnimationDataRegistry)
 	{
 		FRWScopeLock Lock(GAnimationDataRegistry->SkeletalMeshReferencePosesLock, SLT_Write);
 
-		for (auto Iter = GAnimationDataRegistry->SkeletalMeshReferencePoses.CreateIterator(); Iter; ++Iter)
+		if (GAnimationDataRegistry->SkeletalMeshReferencePoses.IsEmpty())
 		{
-			const TObjectKey<USkeletalMeshComponent>& SkeletalMeshComponentPtr = Iter.Key();
-			if (SkeletalMeshComponentPtr.ResolveObjectPtr() == nullptr)
-			{
-				Iter.RemoveCurrent();
-			}
+			// Nothing to clean up
+			return;
+		}
+
+		// Start our cleanup
+		// If cleanup was already in progress, we purge whatever remains (shouldn't happen unless we force GC to run)
+		if (!GAnimationDataRegistry->GCKeysToCleanup.IsEmpty())
+		{
+			GAnimationDataRegistry->RunGCCleanup(INDEX_NONE);
 		}
+
+		if (GAnimationDataRegistry->GCCleanupTickerHandle.IsValid())
+		{
+			FTSTicker::GetCoreTicker().RemoveTicker(GAnimationDataRegistry->GCCleanupTickerHandle);
+		}
+
+		// Copy our keys into a temporary array, we'll process them in small batches over time
+		GAnimationDataRegistry->SkeletalMeshReferencePoses.GenerateKeyArray(GAnimationDataRegistry->GCKeysToCleanup);
+		GAnimationDataRegistry->GCKeysProcessedIndex = 0;
+
+		const float TickerDelay = 0.0f;
+		GAnimationDataRegistry->GCCleanupTickerHandle = FTSTicker::GetCoreTicker().AddTicker(TEXT("AnimNext DataRegistry Post-GC Cleanup"), TickerDelay, GCCleanupTicker);
+	}
+}
+
+bool FDataRegistry::GCCleanupTicker(float DeltaTime)
+{
+	bool bIsDone;
+	{
+		FRWScopeLock Lock(GAnimationDataRegistry->SkeletalMeshReferencePosesLock, SLT_Write);
+		bIsDone = GAnimationDataRegistry->RunGCCleanup(CVar_GCCleanupTimeBudget.GetValueOnAnyThread());
 	}
+
+	const bool bFireTickerAgainAfterDelay = !bIsDone;
+	return bFireTickerAgainAfterDelay;
+}
+
+bool FDataRegistry::RunGCCleanup(double TimeBudget)
+{
+	SCOPED_NAMED_EVENT(AnimNext_DataRegistry_GCCleanup, FColor::Blue);
+
+	if (TimeBudget <= 0.0)
+	{
+		// No budget specified, process everything
+		TimeBudget = UE_DOUBLE_BIG_NUMBER;
+	}
+
+	const double StartTime = FPlatformTime::Seconds();
+
+	while (GCKeysProcessedIndex < GCKeysToCleanup.Num())
+	{
+		const TObjectKey<USkeletalMeshComponent>& SkeletalMeshComponentKey = GCKeysToCleanup[GCKeysProcessedIndex];
+		if (SkeletalMeshComponentKey.ResolveObjectPtr() == nullptr)
+		{
+			SkeletalMeshReferencePoses.Remove(SkeletalMeshComponentKey);
+		}
+
+		GCKeysProcessedIndex++;
+
+		if ((FPlatformTime::Seconds() - StartTime) > TimeBudget)
+		{
+			break;
+		}
+	}
+
+	if (GCKeysProcessedIndex == GCKeysToCleanup.Num())
+	{
+		// We are done
+		GCKeysToCleanup.Empty();
+		GCKeysProcessedIndex = 0;
+		return true;
+	}
+
+	// Not done yet
+	return false;
 }
 
 FDataHandle FDataRegistry::RegisterReferencePose(USkeletalMeshComponent* SkeletalMeshComponent)
diff --git a/Engine/Plugins/Experimental/AnimNext/Source/AnimNext/Public/DataRegistry.h b/Engine/Plugins/Experimental/AnimNext/Source/AnimNext/Public/DataRegistry.h
--- a/Engine/Plugins/Experimental/AnimNext/Source/AnimNext/Public/DataRegistry.h
+++ b/Engine/Plugins/Experimental/AnimNext/Source/AnimNext/Public/DataRegistry.h
@@ -5,6 +5,7 @@
 #include "CoreMinimal.h"
 #include "DataRegistryTypes.h"
 #include "UObject/WeakObjectPtr.h"
+#include "Containers/Ticker.h"
 
 #define UE_API ANIMNEXT_API
 
@@ -186,6 +187,15 @@ private:
 	TMap <TObjectKey<USkeletalMeshComponent>, FReferencePoseData> SkeletalMeshReferencePoses;
 	mutable FRWLock SkeletalMeshReferencePosesLock;
 
+	// A handle to our ticker when GC runs and we need to clean things up
+	FTSTicker::FDelegateHandle GCCleanupTickerHandle;
+
+	// Array of keys to clean up post-GC
+	TArray<TObjectKey<USkeletalMeshComponent>> GCKeysToCleanup;
+
+	// The current index of GC keys to process
+	int32 GCKeysProcessedIndex = 0;
+
 	// Registers a type and sets the allocation block size
 	template<typename DataType>
 	FDataTypeDef RegisterDataType_Impl(int32 AllocationBlockSize)
@@ -222,6 +232,9 @@ private:
 
 	UE_API void ReleaseReferencePoseData();
 
+	static bool GCCleanupTicker(float DeltaTime);
+	bool RunGCCleanup(double TimeBudget);
+
 // --- ---
 private:
 	friend class FAnimNextModuleImpl;
