AnimNext: Pass-by-ref runtime refactor

==Core Changes==

Each AnimNextInstance (graph or module) now contains a 'variable set' structure (FUAFInstanceVariablesData). This contains a number of FUAFInstanceVariablesContainer structures, which wrap either a FInstancePropertyBag (for asset-generated variables) or a FInstancedStruct (for native variables). Each nested scope can reference a FUAFInstanceVariablesContainer from an outer scope if the asset uses that set of shared variables. References to variables are available to inner scopes regardless of how many intervening scopes are present. Intervening scopes can each override individual FUAFInstanceVariablesContainers, partially or as a whole. At present, non-overriden variables in partial overrides will reference memory handles of the outer scope - it may be prudent in future to force these to defaults or copy per-update to maintain correct pass-by-ref semantics.

'Shared' variables now appear in the graph UI with the asset/struct that they come from (existing assets are fixed up on load to the new approach).

Variable overrides are now a seperate type and have a seperate node to generate them. This simplifies the RunGraph node as well as allowing for overrides to be made in Anim Graphs in the future.

Added native struct support to shared variables as well as the existing asset support. Native structs are now a 'first class citizen'.

Added a FAnimNextVariableReference type (and soft reference equivalent) and a picker widget to allow shared variables within specific assets to be referenced directly. Variable references are now a name + (asset/struct) reference, which can be better statically validated at compile time.

Moved & improved the existing state tree RigVM function support into the core (via UAnimNextRigVMAsset/UAnimNextRigVMAssetEditorData).

Variable proxy support has been moved into FAnimNextModuleInstance, rather than be hosted on the component, reducing the public API surface and simplifying code.

UAnimNextComponent variable API has been vastly simplified.

Removed FModuleWriteGuard and module prerequisites. These are not required now due to changes in how injeciton works.

Modified FAnimNextParamType::Get for better codegen and added FromPropertyBagPropertyDesc to make common converisons easier.

Made bool handling for FParamUtils::GetCompatibility consistent with variable type coercion

Key renames:
- DataInterface -> SharedVariables
- AnimNextDataInterfaceInstance -> UAFAssetInstance

==AnimGraph Changes==

Introduced procedural animation graphs. These are UAnimNextAnimationGraphs that are generated at runtime from factory parameters. They use much of the same infrastructure as user-authored graphs, with the execption that they cannot incorporate RigVM logic (no way to author it at runtime). Instead they use direct property-copying to communicate with the graph internals. This first implementation uses direct exposure of the individual trait's data as 'shared variables' to communicate with a single trait stack (actually a few are created in the case dummy refposes are required). In this way, native code can effectively write directly to trait stack variables to control their behavior.

Procedural graphs avoid expensive reconstruction each time using a hashing scheme, caching constructed graphs in a TStripedMap and only regenerating unknown hashes.

Removed FAnimNextDataInterfacePayload in favor of FAnimNextFactoryParams, used to control how procedural graphs are generated.

Removed editor restrictions on FTraitWriter, as this is now used for procedural graphs.

Cached bindings have been removed in favor of the new shared variables system

Removed logic to preload 'factory graphs', including the defualt 'run graph host' - these are now all procedurally generated graphs.

Added IBlendStack::GetGraphInstance to directly access a blend stack's subgraph per-child. Makes for easy access of its variables.

Removed the concept of module dependencies in injection. Now all communication is done directly with the injected graph via the injection request (via queued messages as we dont know the final form of a procedural graph until it is instantiated, so maintaining a game-thread proxy is tricky).

==Aux Plugin Changes==

Updates AnimNextChooser, AnimNextControlRig, AnimNextPoseSearch, AnimNextStateTree, AnimNextWarping, EvaluationNotifies and MovieSceneAnimMixer to the new APIs.

The majority are updated to use the new variable referencing and access APIs
AnimNextStateTree is updated to use the new function infrastructure, plus upgrades to the new variables APIs

diff --git a/Engine/Plugins/Experimental/AnimNext/Config/BaseAnimNext.ini b/Engine/Plugins/Experimental/AnimNext/Config/BaseAnimNext.ini
--- a/Engine/Plugins/Experimental/AnimNext/Config/BaseAnimNext.ini
+++ b/Engine/Plugins/Experimental/AnimNext/Config/BaseAnimNext.ini
@@ -96,9 +96,16 @@
 +ClassRedirects=(OldName="/Script/AnimNextUncookedOnly.AnimNextAnimationGraphEntry", NewName="/Script/AnimNextAnimGraphUncookedOnly.AnimNextAnimationGraphEntry")
 
 ; PlayAnim -> Injection
-+StructRedirects=(OldName="/Script/AnimNext.ModuleInstanceComponent", NewName="/Script/AnimNext.AnimNextModuleInstanceComponent")
++StructRedirects=(OldName="/Script/AnimNext.ModuleInstanceComponent", NewName="/Script/AnimNext.UAFModuleInstanceComponent")
 +StructRedirects=(OldName="/Script/AnimNext.AnimNextGraphInstancePtr", NewName="/Script/AnimNextAnimGraph.AnimNextAnimGraph")
 
-
-
-
+; Pass-by-ref
++StructRedirects=(OldName="/Script/AnimNext.AnimNextDataInterfaceInstance", NewName="/Script/AnimNext.AnimNextInstance")
++ClassRedirects=(OldName="/Script/AnimNext.AnimNextDataInterface", NewName="/Script/AnimNext.AnimNextSharedVariables")
++ClassRedirects=(OldName="/Script/AnimNext.AnimNextDataInterfaceFactory", NewName="/Script/AnimNext.AnimNextSharedVariablesFactory")
++ClassRedirects=(OldName="/Script/AnimNextUncookedOnly.AnimNextDataInterface_EditorData", NewName="/Script/AnimNextUncookedOnly.AnimNextSharedVariables_EditorData")
++ClassRedirects=(OldName="/Script/AnimNextUncookedOnly.AnimNextDataInterfaceEntry", NewName="/Script/AnimNextUncookedOnly.AnimNextSharedVariablesEntry")
++PropertyRedirects=(OldName="/Script/AnimNextUncookedOnly.AnimNextSharedVariablesEntry.DataInterface", NewName="Asset")
++PropertyRedirects=(OldName="/Script/AnimNextUncookedOnly.AnimNextSharedVariablesEntry.DataInterfacePath", NewName="ObjectPath")
++FunctionRedirects=(OldName="AnimNextRigVMAssetLibrary.AddDataInterface",NewName="AnimNextRigVMAssetLibrary.AddSharedVariables")
++FunctionRedirects=(OldName="AnimNextComponent.SetVariable",NewName="AnimNextComponent.BlueprintSetVariable")
diff --git a/Engine/Plugins/Experimental/AnimNext/Source/AnimNext/AnimNext.Build.cs b/Engine/Plugins/Experimental/AnimNext/Source/AnimNext/AnimNext.Build.cs
--- a/Engine/Plugins/Experimental/AnimNext/Source/AnimNext/AnimNext.Build.cs
+++ b/Engine/Plugins/Experimental/AnimNext/Source/AnimNext/AnimNext.Build.cs
@@ -28,6 +28,11 @@ namespace UnrealBuildTool.Rules
 				}
 			);
 
+			if(Target.bWithLiveCoding)
+			{
+				PrivateDependencyModuleNames.Add("LiveCoding");
+			}
+
 			if (Target.bBuildEditor == true)
 			{
 				PrivateDependencyModuleNames.AddRange(
diff --git a/Engine/Plugins/Experimental/AnimNext/Source/AnimNext/Internal/AnimNextExecuteContext.h b/Engine/Plugins/Experimental/AnimNext/Source/AnimNext/Internal/AnimNextExecuteContext.h
--- a/Engine/Plugins/Experimental/AnimNext/Source/AnimNext/Internal/AnimNextExecuteContext.h
+++ b/Engine/Plugins/Experimental/AnimNext/Source/AnimNext/Internal/AnimNextExecuteContext.h
@@ -7,7 +7,7 @@
 #include "Module/AnimNextModuleContextData.h"
 #include "AnimNextExecuteContext.generated.h"
 
-struct FAnimNextDataInterfaceInstance;
+struct FUAFAssetInstance;
 
 namespace UE::AnimNext
 {
diff --git a/Engine/Plugins/Experimental/AnimNext/Source/AnimNext/Internal/Component/AnimNextWorldSubsystem.h b/Engine/Plugins/Experimental/AnimNext/Source/AnimNext/Internal/Component/AnimNextWorldSubsystem.h
--- a/Engine/Plugins/Experimental/AnimNext/Source/AnimNext/Internal/Component/AnimNextWorldSubsystem.h
+++ b/Engine/Plugins/Experimental/AnimNext/Source/AnimNext/Internal/Component/AnimNextWorldSubsystem.h
@@ -35,8 +35,6 @@ struct FModulePendingAction
 		ReleaseHandle,			// Payload = FModulePayloadNone
 		EnableHandle,			// Payload = bool
 		EnableDebugDrawing,		// Payload = bool
-		AddPrerequisite,		// Payload = UE::AnimNext::FModuleHandle (to-be-added handle)
-		RemovePrerequisite,	    // Payload = UE::AnimNext::FModuleHandle (to-be-removed handle)
 	};
 
 	// Marker struct for an 'empty' payload
@@ -106,11 +104,11 @@ private:
 protected:
 #if WITH_EDITOR
 	// Refresh any entries that use the provided module as it has been recompiled.
-	UE_API virtual void OnModuleCompiled(UAnimNextModule* InModule);
+	UE_API void OnModuleCompiled(UAnimNextModule* InModule);
 #endif
-	
+
 	// Register a handle to the subsystem
-	UE_API void RegisterHandle(UE::AnimNext::FModuleHandle& InOutHandle, UAnimNextModule* InModule, UObject* InObject, IAnimNextVariableProxyHost* InProxyHost, EAnimNextModuleInitMethod InitMethod);
+	UE_API void RegisterHandle(UE::AnimNext::FModuleHandle& InOutHandle, UAnimNextModule* InModule, UObject* InObject, EAnimNextModuleInitMethod InitMethod);
 
 	// Unregister a handle from the subsystem
 	UE_API void UnregisterHandle(UE::AnimNext::FModuleHandle& InOutHandle);
@@ -139,16 +137,6 @@ protected:
 	// Queue an input trait event on the supplied handle
 	UE_API void QueueInputTraitEventHandle(UE::AnimNext::FModuleHandle InHandle, FAnimNextTraitEventPtr Event);
 
-	// Add a dependency between equivalent module events for the specified handles
-	// @param	InPrerequisiteHandle	The handle of the module that will be executed first
-	// @param	InSubsequentHandle		The handle of the module that will be executed second
-	UE_API void AddModuleDependencyHandle(UE::AnimNext::FModuleHandle InPrerequisiteHandle, UE::AnimNext::FModuleHandle InSubsequentHandle);
-
-	// Remove a dependency between equivalent module events for the specified handles
-	// @param	InPrerequisiteHandle	The handle of the module that would have been executed first
-	// @param	InSubsequentHandle		The handle of the module that would have been executed second
-	UE_API void RemoveModuleDependencyHandle(UE::AnimNext::FModuleHandle InPrerequisiteHandle, UE::AnimNext::FModuleHandle InSubsequentHandle);
-
 	// Find the const tick function for the specified event. Useful to analyze/log prerequisites. Use AddDependencyHandle, RemoveDependencyHandle to actually modify dependencies
 	// @param	InHandle			The handle of the module
 	// @param	InEventName			The event associated to the wanted tick function
@@ -186,6 +174,24 @@ protected:
 	// @param	InDependency		The kind of dependency to remove
 	UE_API void RemoveModuleEventDependencyHandle(UE::AnimNext::FModuleHandle InHandle, FName InEventName, UE::AnimNext::FModuleHandle OtherHandle, FName OtherEventName, EDependency InDependency);
 
+	// Set the value of the specified variable
+	// @param	InHandle			The handle of the module whose variable we are setting
+	// @param	InVariable			The variable we want to set
+	// @param	InType				The type of the variable we want to set
+	// @param	InData				The data to set the variable with
+	// @return true if the variable was set successfully
+	UE_API EPropertyBagResult SetVariableHandle(UE::AnimNext::FModuleHandle InHandle, const FAnimNextVariableReference& InVariable, const FAnimNextParamType& InType, TConstArrayView<uint8> InData);
+
+	// Accesses the variable of the specified name for writing.
+	// General read access is not provided via this API due to the current double-buffering strategy used to communicate variable writes to worker threads.
+	// This is intended to allow for copy-free writing of larger data structures & arrays, rather than read access. 
+	// @param	InHandle			The handle of the module whose variable we are setting
+	// @param	InVariable			The variable we want to access
+	// @param	InType				The type of the variable we want to access
+	// @param	InFunction			Function that will be called to allow modification of the variables
+	// @return true if the variable was accessed successfully
+	UE_API EPropertyBagResult WriteVariableHandle(UE::AnimNext::FModuleHandle InHandle, const FAnimNextVariableReference& InVariable, const FAnimNextParamType& InType, TFunctionRef<void(TArrayView<uint8>)> InFunction);
+
 protected:
 	friend struct FRigUnit_AnimNextInitializeEvent;
 	friend struct UE::AnimNext::FModuleEventTickFunction;
diff --git a/Engine/Plugins/Experimental/AnimNext/Source/AnimNext/Internal/Graph/RigVMTrait_AnimNextPublicVariables.h b/Engine/Plugins/Experimental/AnimNext/Source/AnimNext/Internal/Graph/RigVMTrait_AnimNextPublicVariables.h
--- a/Engine/Plugins/Experimental/AnimNext/Source/AnimNext/Internal/Graph/RigVMTrait_AnimNextPublicVariables.h
+++ b/Engine/Plugins/Experimental/AnimNext/Source/AnimNext/Internal/Graph/RigVMTrait_AnimNextPublicVariables.h
@@ -2,7 +2,7 @@
 
 #pragma once
 
-#include "DataInterface/AnimNextDataInterfaceHost.h"
+#include "UAFAssetInstance.h"
 #include "RigVMCore/RigVMTrait.h"
 #include "RigVMTrait_AnimNextPublicVariables.generated.h"
 
@@ -10,7 +10,7 @@
 
 class UAnimNextRigVMAsset;
 class FRigVMTraitScope;
-class UAnimNextDataInterface;
+class UAnimNextSharedVariables;
 
 // Represents public variables of an asset via a trait 
 USTRUCT(BlueprintType)
@@ -18,9 +18,9 @@ struct FRigVMTrait_AnimNextPublicVariables : public FRigVMTrait
 {
 	GENERATED_BODY()
 
-	// The data interface that any programmatic pins will be derived from
+	// The asset that any programmatic pins will be derived from, if any
 	UPROPERTY(meta = (Hidden))
-	TObjectPtr<UAnimNextDataInterface> InternalAsset = nullptr;
+	TObjectPtr<UAnimNextSharedVariables> InternalAsset = nullptr;
 
 	// Variable names that are exposed
 	UPROPERTY(meta = (Hidden))
@@ -39,24 +39,4 @@ struct FRigVMTrait_AnimNextPublicVariables : public FRigVMTrait
 #endif
 };
 
-namespace UE::AnimNext
-{
-
-struct FPublicVariablesTraitToDataInterfaceHostAdapter : public IDataInterfaceHost
-{
-	FPublicVariablesTraitToDataInterfaceHostAdapter(const FRigVMTrait_AnimNextPublicVariables& InTrait, const FRigVMTraitScope& InTraitScope)
-		: Trait(InTrait)
-		, TraitScope(InTraitScope)
-	{}
-
-	// IDataInterfaceHost interface
-	ANIMNEXT_API virtual const UAnimNextDataInterface* GetDataInterface() const override;
-	ANIMNEXT_API virtual uint8* GetMemoryForVariable(int32 InVariableIndex, FName InVariableName, const FProperty* InVariableProperty) const override;
-
-	const FRigVMTrait_AnimNextPublicVariables& Trait;
-	const FRigVMTraitScope& TraitScope;
-};
-
-}
-
-#undef UE_API
+#undef UE_API
\ No newline at end of file
diff --git a/Engine/Plugins/Experimental/AnimNext/Source/AnimNext/Internal/Module/AnimNextActorComponentReferenceComponent.h b/Engine/Plugins/Experimental/AnimNext/Source/AnimNext/Internal/Module/AnimNextActorComponentReferenceComponent.h
--- a/Engine/Plugins/Experimental/AnimNext/Source/AnimNext/Internal/Module/AnimNextActorComponentReferenceComponent.h
+++ b/Engine/Plugins/Experimental/AnimNext/Source/AnimNext/Internal/Module/AnimNextActorComponentReferenceComponent.h
@@ -2,7 +2,7 @@
 
 #pragma once
 
-#include "Module/AnimNextModuleInstanceComponent.h"
+#include "Module/UAFModuleInstanceComponent.h"
 #include "Templates/SubclassOf.h"
 
 #include "AnimNextActorComponentReferenceComponent.generated.h"
@@ -13,10 +13,12 @@ class UActorComponent;
 // Note that this does not support dynamic component updates (e.g. if we have multiple, or changing components). It grabs a ptr to the component on
 // initialization and does not update it.
 USTRUCT()
-struct FAnimNextActorComponentReferenceComponent : public FAnimNextModuleInstanceComponent
+struct FAnimNextActorComponentReferenceComponent : public FUAFModuleInstanceComponent
 {
 	GENERATED_BODY()
 
+	FAnimNextActorComponentReferenceComponent() = default;
+
 protected:
 	// Helper function used by derived classes to initialize the component
 	void OnInitializeHelper(UScriptStruct* InScriptStruct);
diff --git a/Engine/Plugins/Experimental/AnimNext/Source/AnimNext/Internal/Module/AnimNextModule.h b/Engine/Plugins/Experimental/AnimNext/Source/AnimNext/Internal/Module/AnimNextModule.h
--- a/Engine/Plugins/Experimental/AnimNext/Source/AnimNext/Internal/Module/AnimNextModule.h
+++ b/Engine/Plugins/Experimental/AnimNext/Source/AnimNext/Internal/Module/AnimNextModule.h
@@ -3,7 +3,7 @@
 #pragma once
 
 #include "CoreMinimal.h"
-#include "DataInterface/AnimNextDataInterface.h"
+#include "Variables/AnimNextSharedVariables.h"
 #include "RigVMCore/RigVM.h"
 #include "RigVMHost.h"
 #include "Graph/AnimNextGraphState.h"
@@ -22,6 +22,7 @@ struct FAnimNextGraphInstance;
 struct FAnimNextScheduleGraphTask;
 struct FAnimNextEditorParam;
 struct FAnimNextParam;
+struct FAnimNextModuleInstance;
 
 namespace UE::AnimNext
 {
@@ -45,7 +46,7 @@ namespace UE::AnimNext::Editor
 
 // Root asset represented by a component when instantiated
 UCLASS(MinimalAPI, BlueprintType)
-class UAnimNextModule : public UAnimNextDataInterface
+class UAnimNextModule : public UAnimNextSharedVariables
 {
 	GENERATED_BODY()
 
diff --git a/Engine/Plugins/Experimental/AnimNext/Source/AnimNext/Internal/Module/AnimNextModuleContextData.h b/Engine/Plugins/Experimental/AnimNext/Source/AnimNext/Internal/Module/AnimNextModuleContextData.h
--- a/Engine/Plugins/Experimental/AnimNext/Source/AnimNext/Internal/Module/AnimNextModuleContextData.h
+++ b/Engine/Plugins/Experimental/AnimNext/Source/AnimNext/Internal/Module/AnimNextModuleContextData.h
@@ -3,7 +3,7 @@
 #pragma once
 
 #include "CoreMinimal.h"
-#include "DataInterface/AnimNextDataInterfaceInstance.h"
+#include "UAFAssetInstance.h"
 #include "AnimNextModuleContextData.generated.h"
 
 struct FAnimNextModuleInstance;
@@ -17,7 +17,7 @@ struct FAnimNextModuleContextData
 
 	ANIMNEXT_API explicit FAnimNextModuleContextData(FAnimNextModuleInstance* InModuleInstance);
 
-	ANIMNEXT_API FAnimNextModuleContextData(FAnimNextModuleInstance* InModuleInstance, const FAnimNextDataInterfaceInstance* InDataInterfaceInstance);
+	ANIMNEXT_API FAnimNextModuleContextData(FAnimNextModuleInstance* InModuleInstance, const FUAFAssetInstance* InInstance);
 
 	// Get the object that the module instance is bound to, if any
 	ANIMNEXT_API UObject* GetObject() const;
@@ -28,18 +28,18 @@ struct FAnimNextModuleContextData
 		return *ModuleInstance;
 	}
 
-	const FAnimNextDataInterfaceInstance& GetDataInterfaceInstance() const
+	const FUAFAssetInstance& GetInstance() const
 	{
-		check(DataInterfaceInstance != nullptr);
-		return *DataInterfaceInstance;
+		check(Instance != nullptr);
+		return *Instance;
 	}
 
 private:
 	// Module instance that is currently executing.
 	FAnimNextModuleInstance* ModuleInstance = nullptr;
 
-	// Data interface that is currently executing. Can be the same as ModuleInstance
-	const FAnimNextDataInterfaceInstance* DataInterfaceInstance = nullptr;
+	// Instance that is currently executing. Can be the same as ModuleInstance
+	const FUAFAssetInstance* Instance = nullptr;
 
 	friend class UAnimNextModule;
 	friend struct FAnimNextExecuteContext;
diff --git a/Engine/Plugins/Experimental/AnimNext/Source/AnimNext/Internal/Module/AnimNextModuleInstance.h b/Engine/Plugins/Experimental/AnimNext/Source/AnimNext/Internal/Module/AnimNextModuleInstance.h
--- a/Engine/Plugins/Experimental/AnimNext/Source/AnimNext/Internal/Module/AnimNextModuleInstance.h
+++ b/Engine/Plugins/Experimental/AnimNext/Source/AnimNext/Internal/Module/AnimNextModuleInstance.h
@@ -5,13 +5,13 @@
 #include "CoreMinimal.h"
 #include "AnimNextModule.h"
 #include "Module/ModuleHandle.h"
-#include "Component/AnimNextPublicVariablesProxy.h"
-#include "DataInterface/AnimNextDataInterfaceInstance.h"
+#include "Variables/UAFInstanceVariableDataProxy.h"
+#include "UAFAssetInstance.h"
 #include "Module/ModuleTickFunction.h"
 #include "TraitCore/TraitEvent.h"
 #include "TraitCore/TraitEventList.h"
 #include "Module/ModuleHandle.h"
-#include "Module/AnimNextModuleInstanceComponent.h"
+#include "Module/UAFModuleInstanceComponent.h"
 #include "StructUtils/InstancedStruct.h"
 #include "Module/AnimNextModuleInitMethod.h"
 #include "Module/TaskRunLocation.h"
@@ -25,15 +25,11 @@ struct FRigUnit_AnimNextInitializeEvent;
 struct FRigUnit_CopyModuleProxyVariables;
 struct FAnimNextNotifyDispatcherComponent;
 struct FAnimNextGraphInstance;
-class IAnimNextVariableProxyHost;
 
 namespace UE::AnimNext
 {
-	struct FProxyVariablesContext;
 	struct FModuleEventTickFunction;
 	struct FModuleEndTickFunction;
-	struct FModuleInjectionDataInterfaceAdapter;
-	struct FModuleWriteGuard;
 }
 
 namespace UE::AnimNext::Debug
@@ -43,12 +39,17 @@ namespace UE::AnimNext::Debug
 
 namespace UE::AnimNext
 {
-	using ModuleInstanceComponentMapType = TMap<FName, TInstancedStruct<FAnimNextModuleInstanceComponent>>;
+	using ModuleInstanceComponentMapType = TMap<FName, TInstancedStruct<FUAFModuleInstanceComponent>>;
+}
+
+namespace UE::AnimNext::Tests
+{
+	class FVariables_Shared;
 }
 
 // Root memory owner of a parameterized schedule 
 USTRUCT()
-struct FAnimNextModuleInstance : public FAnimNextDataInterfaceInstance
+struct FAnimNextModuleInstance : public FUAFAssetInstance
 {
 	GENERATED_BODY()
 
@@ -57,7 +58,6 @@ struct FAnimNextModuleInstance : public FAnimNextDataInterfaceInstance
 		UAnimNextModule* InModule,
 		UObject* InObject,
 		UE::AnimNext::TPool<FAnimNextModuleInstance>* InPool,
-		IAnimNextVariableProxyHost* InProxyHost,
 		EAnimNextModuleInitMethod InInitMethod);
 	ANIMNEXT_API ~FAnimNextModuleInstance();
 
@@ -81,22 +81,6 @@ struct FAnimNextModuleInstance : public FAnimNextDataInterfaceInstance
 	// Get the module that this instance represents
 	ANIMNEXT_API const UAnimNextModule* GetModule() const;
 
-	// Returns a typed module instance component, creating it lazily the first time it is queried
-	template<class ComponentType>
-	ComponentType& GetComponent();
-
-	// Returns a typed module instance component pointer if found or nullptr otherwise
-	template<class ComponentType>
-	ComponentType* TryGetComponent();
-
-	// Returns a typed module instance component pointer if found or nullptr otherwise
-	template<class ComponentType>
-	const ComponentType* TryGetComponent() const;
-
-	// Returns a pointer to the specified component, or nullptr if not found
-	ANIMNEXT_API const FAnimNextModuleInstanceComponent* TryGetComponent(int32 ComponentNameHash, FName ComponentName) const;
-	ANIMNEXT_API FAnimNextModuleInstanceComponent* TryGetComponent(int32 ComponentNameHash, FName ComponentName);
-
 #if UE_ENABLE_DEBUG_DRAWING
 	// Get the debug draw interface
 	ANIMNEXT_API FRigVMDrawInterface* GetDebugDrawInterface();
@@ -141,10 +125,7 @@ struct FAnimNextModuleInstance : public FAnimNextDataInterfaceInstance
 
 private:
 	// Setup the entry
-	void Initialize();
-
-	// Tear-down the entry
-	void Uninitialize();
+	ANIMNEXT_API void Initialize();
 
 	// Clear data that binds the schedule to a runtime (e.g. tick functions) and any instance data
 	void ResetBindingsAndInstanceData();
@@ -172,17 +153,11 @@ private:
 	// Called in the module end tick function before any events are dispatched.
 	void RaiseTraitEvents(const UE::AnimNext::FTraitEventList& EventList);
 
-	// Add a prerequisite instance
-	void AddPrerequisite(FAnimNextModuleInstance& InPrerequisiteInstance);
-
-	// Remove a prerequisite instance
-	void RemovePrerequisite(FAnimNextModuleInstance& InPrerequisiteInstance);
+	// Set the value of the specified variable via the proxy double-buffer. Called from any thread.
+	EPropertyBagResult SetProxyVariable(const FAnimNextVariableReference& InVariable, const FAnimNextParamType& InType, TConstArrayView<uint8> InData);
 
-	// @return true if the instance is a prerequisite instance
-	ANIMNEXT_API bool IsPrerequisite(const FAnimNextModuleInstance& InPrerequisiteInstance) const;
-
-	// Call the supplied function for each prerequisite module
-	ANIMNEXT_API void ForEachPrerequisite(TFunctionRef<void(FAnimNextModuleInstance& InPrerequisiteInstance)>) const;
+	// Accesses value of the specified variable via the proxy double-buffer. Called from any thread.
+	EPropertyBagResult WriteProxyVariable(const FAnimNextVariableReference& InVariable, const FAnimNextParamType& InType, TFunctionRef<void(TArrayView<uint8>)> InFunction);
 
 private:
 	friend UE::AnimNext::FModuleEventTickFunction;
@@ -191,12 +166,10 @@ private:
 	friend FRigUnit_AnimNextRunAnimationGraph_v2;
 	friend UAnimNextWorldSubsystem;
 	friend FRigUnit_AnimNextInitializeEvent;
-	friend UE::AnimNext::FProxyVariablesContext;
 	friend FRigUnit_CopyModuleProxyVariables;
-	friend UE::AnimNext::FModuleInjectionDataInterfaceAdapter;
-	friend UE::AnimNext::FModuleWriteGuard;
 	friend FAnimNextNotifyDispatcherComponent;
 	friend FAnimNextGraphInstance;
+	friend UE::AnimNext::Tests::FVariables_Shared;
 
 	// Object this entry is bound to
 	UPROPERTY(Transient)
@@ -205,45 +178,30 @@ private:
 	// The pool that this module instance exists in
 	UE::AnimNext::TPool<FAnimNextModuleInstance>* Pool = nullptr;
 
-	// Seperate proxy host ptr - TODO: this only exists because of UHT not parsing Verse native classes right now.
-	// SHould be removed and replaced with a cast of the Cast<IAnimNextVariableProxyHost>(Object) when VNI gets replaced.
-	IAnimNextVariableProxyHost* ProxyHost = nullptr;
-
 	// Copy of the handle that represents this entry to client systems
 	UE::AnimNext::FModuleHandle Handle;
 
 	// Pre-allocated graph of tick functions
 	TArray<UE::AnimNext::FModuleEventTickFunction> TickFunctions;
 
-	// Reference-counted prerequisites. Prerequisites can be requested/unrequested multiple times, so are only truly removed when references are zero
-	struct FPrerequisiteReference
-	{
-		UE::AnimNext::FModuleHandle Handle;
-		int32 ReferenceCount = 0;
-	};
-
-	// All pre-requisite references, only modified from the game thread
-	TArray<FPrerequisiteReference> PrerequisiteRefs;
-
-	// All subsequent references, only modified from the game thread
-	TArray<UE::AnimNext::FModuleHandle> SubsequentRefs;
-
 	// Input event list to be processed on the next update
 	UE::AnimNext::FTraitEventList InputEventList;
 
 	// Output event list to be processed at the end of the schedule tick
 	UE::AnimNext::FTraitEventList OutputEventList;
 
-	// Access detector to ensure module prerequisite and general execution access is safe
-	// This protects all calls & tick functions and prevents prerequisites from running unrelated tick functions concurrently
-	UE_MT_DECLARE_RW_RECURSIVE_ACCESS_DETECTOR(AccessDetector);
-
 	// Lock to ensure event list actions are thread safe
 	FRWLock EventListLock;
 
-	// Proxy public variables
+	// Lock to ensure proxy access is thread safe
+	FRWLock ProxyLock;
+
+	// Proxy public variables, double-buffered
 	UPROPERTY(Transient)
-	FAnimNextPublicVariablesProxy PublicVariablesProxy;
+	FUAFInstanceVariableDataProxy ProxyVariables[2];
+
+	// Index into current buffer to write the proxy variables
+	int32 ProxyWriteIndex = 0;
 
 #if UE_ENABLE_DEBUG_DRAWING
 	// Debug draw support
@@ -251,10 +209,6 @@ private:
 	TUniquePtr<UE::AnimNext::Debug::FDebugDraw> DebugDraw;
 #endif
 
-	// Components of this module
-	UPROPERTY()
-	TMap<FName, TInstancedStruct<FAnimNextModuleInstanceComponent>> ComponentMap;
-
 	enum class EInitState : uint8
 	{
 		NotInitialized,
@@ -305,9 +259,6 @@ private:
 #endif
 
 private:
-	// Adds the specified component and returns a reference to it
-	ANIMNEXT_API FAnimNextModuleInstanceComponent& AddComponentInternal(int32 ComponentNameHash, FName ComponentName, TInstancedStruct<FAnimNextModuleInstanceComponent>&& Component);
-
 #if ANIMNEXT_TRACE_ENABLED
 	bool bTracedThisFrame = false;
 #endif
@@ -322,40 +273,6 @@ struct TStructOpsTypeTraits<FAnimNextModuleInstance> : public TStructOpsTypeTrai
 	};
 };
 
-template<typename ComponentType>
-ComponentType& FAnimNextModuleInstance::GetComponent()
-{
-	const FName ComponentName = ComponentType::StaticStruct()->GetFName();
-	const int32 ComponentNameHash = GetTypeHash(ComponentName);
-
-	if (ComponentType* Component = static_cast<ComponentType*>(TryGetComponent(ComponentNameHash, ComponentName)))
-	{
-		return *Component;
-	}
-
-	TInstancedStruct<ComponentType> ComponentStruct = TInstancedStruct<ComponentType>::Make();
-	ComponentStruct.template GetMutable<ComponentType>().Initialize(*this);
-	return static_cast<ComponentType&>(AddComponentInternal(ComponentNameHash, ComponentName, MoveTemp(ComponentStruct)));
-}
-
-template<class ComponentType>
-ComponentType* FAnimNextModuleInstance::TryGetComponent()
-{
-	const FName ComponentName = ComponentType::StaticStruct()->GetFName();
-	const int32 ComponentNameHash = GetTypeHash(ComponentName);
-
-	return static_cast<ComponentType*>(TryGetComponent(ComponentNameHash, ComponentName));
-}
-
-template<class ComponentType>
-const ComponentType* FAnimNextModuleInstance::TryGetComponent() const
-{
-	const FName ComponentName = ComponentType::StaticStruct()->GetFName();
-	const int32 ComponentNameHash = GetTypeHash(ComponentName);
-
-	return static_cast<ComponentType*>(TryGetComponent(ComponentNameHash, ComponentName));
-}
-
 #if ANIMNEXT_TRACE_ENABLED
 #define TRACE_ANIMNEXT_MODULE(ModuleInstance) ModuleInstance.Trace();
 #else
diff --git a/Engine/Plugins/Experimental/AnimNext/Source/AnimNext/Internal/Module/AnimNextSkeletalMeshComponentReferenceComponent.h b/Engine/Plugins/Experimental/AnimNext/Source/AnimNext/Internal/Module/AnimNextSkeletalMeshComponentReferenceComponent.h
--- a/Engine/Plugins/Experimental/AnimNext/Source/AnimNext/Internal/Module/AnimNextSkeletalMeshComponentReferenceComponent.h
+++ b/Engine/Plugins/Experimental/AnimNext/Source/AnimNext/Internal/Module/AnimNextSkeletalMeshComponentReferenceComponent.h
@@ -16,10 +16,6 @@ struct FAnimNextSkeletalMeshComponentReferenceComponent : public FAnimNextActorC
 
 	UE_API FAnimNextSkeletalMeshComponentReferenceComponent();
 	UE_API USkeletalMeshComponent* GetComponent() const;
-
-private:
-	// FAnimNextModuleInstanceComponent interface
-	UE_API virtual void OnInitialize() override;
 };
 
 #undef UE_API
diff --git a/Engine/Plugins/Experimental/AnimNext/Source/AnimNext/Private/AnimNextRigVMAsset.cpp b/Engine/Plugins/Experimental/AnimNext/Source/AnimNext/Private/AnimNextRigVMAsset.cpp
--- a/Engine/Plugins/Experimental/AnimNext/Source/AnimNext/Private/AnimNextRigVMAsset.cpp
+++ b/Engine/Plugins/Experimental/AnimNext/Source/AnimNext/Private/AnimNextRigVMAsset.cpp
@@ -1,8 +1,15 @@
 // Copyright Epic Games, Inc. All Rights Reserved.
 
 #include "AnimNextRigVMAsset.h"
+
+#if WITH_LIVE_CODING
+#include "ILiveCodingModule.h"
+#endif
+#include "AnimNextFunctionHandle.h"
+#include "AnimNextRigVMFunctionData.h"
 #include "RigVMRuntimeDataRegistry.h"
 #include "Graph/RigUnit_AnimNextBeginExecution.h"
+#include "Modules/ModuleManager.h"
 #include "UObject/AssetRegistryTagsContext.h"
 
 UAnimNextRigVMAsset::UAnimNextRigVMAsset(const FObjectInitializer& ObjectInitializer)
@@ -81,19 +88,20 @@ TArray<FRigVMExternalVariable> UAnimNextRigVMAsset::GetExternalVariablesImpl(boo
 {
 	TArray<FRigVMExternalVariable> ExternalVariables;
 
-	if(const UPropertyBag* PropertyBag = VariableDefaults.GetPropertyBagStruct())
+	if(const UPropertyBag* PropertyBag = CombinedPropertyBag.GetPropertyBagStruct())
 	{
 		TConstArrayView<FPropertyBagPropertyDesc> VariableDescs = PropertyBag->GetPropertyDescs();
-		if(VariableDescs.Num() > 0)
+		if(VariableDescs.Num() != 0)
 		{
-			ExternalVariables.Reserve(VariableDescs.Num());
+			uint8* Container = const_cast<uint8*>(CombinedPropertyBag.GetValue().GetMemory());
+			ExternalVariables.Reserve(ExternalVariables.Num() + VariableDescs.Num());
 			for(const FPropertyBagPropertyDesc& Desc : VariableDescs)
 			{
 				const FProperty* Property = Desc.CachedProperty;
-				FRigVMExternalVariable ExternalVariable = FRigVMExternalVariable::Make(Property, const_cast<uint8*>(VariableDefaults.GetValue().GetMemory()));
+				FRigVMExternalVariable ExternalVariable = FRigVMExternalVariable::Make(Property, Container);
 				if(!ExternalVariable.IsValid())
 				{
-					UE_LOG(LogRigVM, Warning, TEXT("%s: Property '%s' of type '%s' is not supported."), *GetClass()->GetName(), *Property->GetName(), *Property->GetCPPType());
+					UE_LOG(LogRigVM, Error, TEXT("%s: Property '%s' of type '%s' is not supported."), *GetName(), *Property->GetName(), *Property->GetCPPType());
 					continue;
 				}
 
@@ -104,3 +112,70 @@ TArray<FRigVMExternalVariable> UAnimNextRigVMAsset::GetExternalVariablesImpl(boo
 
 	return ExternalVariables;
 }
+
+
+void UAnimNextRigVMAsset::CallFunctionHandle(UE::AnimNext::FFunctionHandle InHandle, FRigVMExtendedExecuteContext& InContext, FInstancedPropertyBag& InArgs)
+{
+#if WITH_EDITOR
+	check(InHandle.IsValidForVM(VM->GetVMHash()));
+#endif
+
+	const FAnimNextRigVMFunctionData& Function = FunctionData[InHandle.FunctionIndex];
+	const int32 NumArgs = Function.ArgIndices.Num();
+	check(InArgs.GetNumPropertiesInBag() == NumArgs);
+
+	const UPropertyBag* PropertyBag = InArgs.GetPropertyBagStruct();
+	if (PropertyBag)
+	{
+		// Copy-in args
+		void* ContainerPtr = InArgs.GetMutableValue().GetMemory();
+		TConstArrayView<FPropertyBagPropertyDesc> Descs = PropertyBag->GetPropertyDescs();
+		for (int32 ArgIndex = 0; ArgIndex < NumArgs; ++ArgIndex)
+		{
+			const FPropertyBagPropertyDesc& Desc = Descs[ArgIndex];
+			const int32 ExternalVariableIndex = Function.ArgIndices[ArgIndex];
+			check(VM->GetExternalVariableDefs()[ExternalVariableIndex].Property->GetClass() == Desc.CachedProperty->GetClass());
+			void* DestPtr = InContext.ExternalVariableRuntimeData[ExternalVariableIndex].Memory;
+			void* SrcPtr = Desc.CachedProperty->ContainerPtrToValuePtr<void>(ContainerPtr);
+			Desc.CachedProperty->CopyCompleteValue(DestPtr, SrcPtr);
+		}
+
+		// Run the function's wrapper event
+		VM->ExecuteVM(InContext, Function.EventName);
+
+		// Copy-out args
+		for (int32 ArgIndex = 0; ArgIndex < NumArgs; ++ArgIndex)
+		{
+			const FPropertyBagPropertyDesc& Desc = Descs[ArgIndex];
+			void* DestPtr = Desc.CachedProperty->ContainerPtrToValuePtr<void>(ContainerPtr);
+			void* SrcPtr = InContext.ExternalVariableRuntimeData[Function.ArgIndices[ArgIndex]].Memory;
+			Desc.CachedProperty->CopyCompleteValue(DestPtr, SrcPtr);
+		}
+	}
+	else
+	{
+		// No args - just run the function's wrapper event
+		VM->ExecuteVM(InContext, Function.EventName);
+	}
+}
+
+UE::AnimNext::FFunctionHandle UAnimNextRigVMAsset::GetFunctionHandle(FName InEventName) const
+{
+	using namespace UE::AnimNext;
+
+	for (int32 FunctionIndex = 0; FunctionIndex < FunctionData.Num(); ++FunctionIndex)
+	{
+		const FAnimNextRigVMFunctionData& Function = FunctionData[FunctionIndex];
+		if (Function.EventName == InEventName)
+		{
+			FFunctionHandle Handle;
+			Handle.FunctionIndex = FunctionIndex;
+#if WITH_EDITOR
+			Handle.VMHash = VM->GetVMHash();
+#endif
+			return Handle;
+		}
+	}
+
+	return FFunctionHandle();
+}
\ No newline at end of file
diff --git a/Engine/Plugins/Experimental/AnimNext/Source/AnimNext/Private/Component/AnimNextComponent.cpp b/Engine/Plugins/Experimental/AnimNext/Source/AnimNext/Private/Component/AnimNextComponent.cpp
--- a/Engine/Plugins/Experimental/AnimNext/Source/AnimNext/Private/Component/AnimNextComponent.cpp
+++ b/Engine/Plugins/Experimental/AnimNext/Source/AnimNext/Private/Component/AnimNextComponent.cpp
@@ -6,7 +6,7 @@
 #include "Blueprint/BlueprintExceptionInfo.h"
 #include "Engine/World.h"
 #include "Module/ModuleTaskContext.h"
-#include "Module/ProxyVariablesContext.h"
+#include "Variables/AnimNextVariableReference.h"
 
 void UAnimNextComponent::OnRegister()
 {
@@ -20,7 +20,6 @@ void UAnimNextComponent::OnRegister()
 	{
 		check(!ModuleHandle.IsValid());
 
-		CreatePublicVariablesProxy();
 		Subsystem->Register(this);
 #if UE_ENABLE_DEBUG_DRAWING
 		Subsystem->ShowDebugDrawing(this, bShowDebugDrawing);
@@ -36,7 +35,6 @@ void UAnimNextComponent::OnUnregister()
 	{
 		Subsystem->Unregister(this);
 		Subsystem = nullptr;
-		DestroyPublicVariablesProxy();
 	}
 }
 
@@ -61,48 +59,12 @@ void UAnimNextComponent::EndPlay(const EEndPlayReason::Type EndPlayReason)
 	SetEnabled(false);
 }
 
-#if WITH_EDITOR
-void UAnimNextComponent::OnModuleCompiled()
-{
-	CreatePublicVariablesProxy();
-}
-#endif
-
-void UAnimNextComponent::CreatePublicVariablesProxy()
-{
-	PublicVariablesProxyMap.Reset();
-	PublicVariablesProxy.Reset();
-	if(Module && Module->GetPublicVariableDefaults().GetPropertyBagStruct())
-	{
-		PublicVariablesProxy.Data = Module->GetPublicVariableDefaults();
-		TConstArrayView<FPropertyBagPropertyDesc> ProxyDescs = PublicVariablesProxy.Data.GetPropertyBagStruct()->GetPropertyDescs();
-		PublicVariablesProxyMap.Reserve(ProxyDescs.Num());
-		for(int32 DescIndex = 0; DescIndex < ProxyDescs.Num(); ++DescIndex)
-		{
-			PublicVariablesProxyMap.Add(ProxyDescs[DescIndex].Name, DescIndex);
-		}
-		PublicVariablesProxy.DirtyFlags.SetNum(ProxyDescs.Num(), false);
-	}
-}
-
-void UAnimNextComponent::DestroyPublicVariablesProxy()
-{
-	PublicVariablesProxyMap.Empty();
-	PublicVariablesProxy.Empty();
-}
-
-void UAnimNextComponent::FlipPublicVariablesProxy(const UE::AnimNext::FProxyVariablesContext& InContext)
-{
-	FRWScopeLock Lock(PublicVariablesLock, SLT_Write);
-	Swap(InContext.GetPublicVariablesProxy(), PublicVariablesProxy);
-}
-
-void UAnimNextComponent::SetVariable(FName Name, int32 Value)
+void UAnimNextComponent::BlueprintSetVariable(FName Name, int32 Value)
 {
 	checkNoEntry();
 }
 
-DEFINE_FUNCTION(UAnimNextComponent::execSetVariable)
+DEFINE_FUNCTION(UAnimNextComponent::execBlueprintSetVariable)
 {
 	using namespace UE::AnimNext;
 
@@ -140,737 +102,175 @@ DEFINE_FUNCTION(UAnimNextComponent::execSetVariable)
 		return;
 	}
 
-	int32* IndexPtr = P_THIS->PublicVariablesProxyMap.Find(Name);
-	if (IndexPtr == nullptr)
-	{
-		FBlueprintExceptionInfo ExceptionInfo(
-			EBlueprintExceptionType::NonFatalError,
-			FText::Format(NSLOCTEXT("UAFComponent", "UAFComponent_SetVariableNotFoundWarning", "Unknown variable name '{0}' supplied to Set Variable"), FText::FromName(Name))
-		);
-
-		FBlueprintCoreDelegates::ThrowScriptException(P_THIS, Stack, ExceptionInfo);
-		return;
-	}
-
-	bool bIsTypeMatching = false;
+	EPropertyBagResult Result = EPropertyBagResult::Success;
 
 	P_NATIVE_BEGIN;
 
 	{
-		FRWScopeLock Lock(P_THIS->PublicVariablesLock, SLT_Write);
-
-		TConstArrayView<FPropertyBagPropertyDesc> ProxyDescs = P_THIS->PublicVariablesProxy.Data.GetPropertyBagStruct()->GetPropertyDescs();
-		const FProperty* CachedProperty = ProxyDescs[*IndexPtr].CachedProperty;
-
-		bIsTypeMatching = CachedProperty->GetClass() == ValueProp->GetClass();
-
-		if (!bIsTypeMatching)
-		{
-			if (const FByteProperty* SourceByteProperty = CastField<FByteProperty>(ValueProp))
-			{
-				if (const FEnumProperty* TargetEnumProperty = CastField<FEnumProperty>(CachedProperty))
-				{
-					bIsTypeMatching = TargetEnumProperty->GetUnderlyingProperty()->IsA<FByteProperty>() || SourceByteProperty->Enum == TargetEnumProperty->GetEnum();
-				}
-			}
-		}
-
-		if (bIsTypeMatching)
-		{
-			const void* ValuePtr = ValueProp->ContainerPtrToValuePtr<void>(ContainerPtr);
-			CachedProperty->SetValue_InContainer(P_THIS->PublicVariablesProxy.Data.GetMutableValue().GetMemory(), ValuePtr);
-			P_THIS->PublicVariablesProxy.DirtyFlags[*IndexPtr] = true;
-			P_THIS->PublicVariablesProxy.bIsDirty = true;
-		}
+		FAnimNextParamType Type = FAnimNextParamType::FromProperty(ValueProp);
+		const uint8* ValuePtr = ValueProp->ContainerPtrToValuePtr<uint8>(ContainerPtr);
+		PRAGMA_DISABLE_DEPRECATION_WARNINGS
+		Result = P_THIS->SetVariableInternal(FAnimNextVariableReference(Name), Type, TConstArrayView<uint8>(ValuePtr, ValueProp->GetElementSize()));
+		PRAGMA_ENABLE_DEPRECATION_WARNINGS
 	}
 
 	P_NATIVE_END;
 
-	if (!bIsTypeMatching)
-	{
-		FBlueprintExceptionInfo ExceptionInfo(
-			EBlueprintExceptionType::NonFatalError,
-			FText::Format(NSLOCTEXT("UAFComponent", "UAFComponent_SetVariableTypeMismatch", "Incompatible type supplied for variable '{0}'"), FText::FromName(Name))
-		);
-
-		FBlueprintCoreDelegates::ThrowScriptException(P_THIS, Stack, ExceptionInfo);
-		return;
-	}
-}
-
-TOptional<bool> UAnimNextComponent::GetVariableBool(const FName Name) const
-{
-	const int32* IndexPtr = PublicVariablesProxyMap.Find(Name);
-	if (IndexPtr == nullptr)
-	{
-		return TOptional<bool>();
-	}
-
-	FRWScopeLock Lock(PublicVariablesLock, SLT_ReadOnly);
-
-	TConstArrayView<FPropertyBagPropertyDesc> ProxyDescs = PublicVariablesProxy.Data.GetPropertyBagStruct()->GetPropertyDescs();
-	const FPropertyBagPropertyDesc& PropertyDesc = ProxyDescs[*IndexPtr];
-
-	TValueOrError<bool, EPropertyBagResult> Result = PublicVariablesProxy.Data.GetValueBool(PropertyDesc);
-	return Result.HasValue() ? TOptional<bool>(Result.GetValue()) : TOptional<bool>();
-}
-
-TOptional<uint8> UAnimNextComponent::GetVariableByte(const FName Name) const
-{
-	const int32* IndexPtr = PublicVariablesProxyMap.Find(Name);
-	if (IndexPtr == nullptr)
-	{
-		return TOptional<uint8>();
-	}
-
-	FRWScopeLock Lock(PublicVariablesLock, SLT_ReadOnly);
-
-	TConstArrayView<FPropertyBagPropertyDesc> ProxyDescs = PublicVariablesProxy.Data.GetPropertyBagStruct()->GetPropertyDescs();
-	const FPropertyBagPropertyDesc& PropertyDesc = ProxyDescs[*IndexPtr];
-
-	TValueOrError<uint8, EPropertyBagResult> Result = PublicVariablesProxy.Data.GetValueByte(PropertyDesc);
-	return Result.HasValue() ? TOptional<uint8>(Result.GetValue()) : TOptional<uint8>();
-}
-
-TOptional<int32> UAnimNextComponent::GetVariableInt32(const FName Name) const
-{
-	const int32* IndexPtr = PublicVariablesProxyMap.Find(Name);
-	if (IndexPtr == nullptr)
-	{
-		return TOptional<int32>();
-	}
-
-	FRWScopeLock Lock(PublicVariablesLock, SLT_ReadOnly);
-
-	TConstArrayView<FPropertyBagPropertyDesc> ProxyDescs = PublicVariablesProxy.Data.GetPropertyBagStruct()->GetPropertyDescs();
-	const FPropertyBagPropertyDesc& PropertyDesc = ProxyDescs[*IndexPtr];
-
-	TValueOrError<int32, EPropertyBagResult> Result = PublicVariablesProxy.Data.GetValueInt32(PropertyDesc);
-	return Result.HasValue() ? TOptional<int32>(Result.GetValue()) : TOptional<int32>();
-}
-
-TOptional<uint32> UAnimNextComponent::GetVariableUInt32(const FName Name) const
-{
-	const int32* IndexPtr = PublicVariablesProxyMap.Find(Name);
-	if (IndexPtr == nullptr)
-	{
-		return TOptional<uint32>();
-	}
-
-	FRWScopeLock Lock(PublicVariablesLock, SLT_ReadOnly);
-
-	TConstArrayView<FPropertyBagPropertyDesc> ProxyDescs = PublicVariablesProxy.Data.GetPropertyBagStruct()->GetPropertyDescs();
-	const FPropertyBagPropertyDesc& PropertyDesc = ProxyDescs[*IndexPtr];
-
-	TValueOrError<uint32, EPropertyBagResult> Result = PublicVariablesProxy.Data.GetValueUInt32(PropertyDesc);
-	return Result.HasValue() ? TOptional<uint32>(Result.GetValue()) : TOptional<uint32>();
-}
-
-TOptional<int64> UAnimNextComponent::GetVariableInt64(const FName Name) const
-{
-	const int32* IndexPtr = PublicVariablesProxyMap.Find(Name);
-	if (IndexPtr == nullptr)
-	{
-		return TOptional<int64>();
-	}
-
-	FRWScopeLock Lock(PublicVariablesLock, SLT_ReadOnly);
-
-	TConstArrayView<FPropertyBagPropertyDesc> ProxyDescs = PublicVariablesProxy.Data.GetPropertyBagStruct()->GetPropertyDescs();
-	const FPropertyBagPropertyDesc& PropertyDesc = ProxyDescs[*IndexPtr];
-
-	TValueOrError<int64, EPropertyBagResult> Result = PublicVariablesProxy.Data.GetValueInt64(PropertyDesc);
-	return Result.HasValue() ? TOptional<int64>(Result.GetValue()) : TOptional<int64>();
-}
-
-TOptional<uint64> UAnimNextComponent::GetVariableUInt64(const FName Name) const
-{
-	const int32* IndexPtr = PublicVariablesProxyMap.Find(Name);
-	if (IndexPtr == nullptr)
-	{
-		return TOptional<uint64>();
-	}
-
-	FRWScopeLock Lock(PublicVariablesLock, SLT_ReadOnly);
-
-	TConstArrayView<FPropertyBagPropertyDesc> ProxyDescs = PublicVariablesProxy.Data.GetPropertyBagStruct()->GetPropertyDescs();
-	const FPropertyBagPropertyDesc& PropertyDesc = ProxyDescs[*IndexPtr];
-
-	TValueOrError<uint64, EPropertyBagResult> Result = PublicVariablesProxy.Data.GetValueUInt64(PropertyDesc);
-	return Result.HasValue() ? TOptional<uint64>(Result.GetValue()) : TOptional<uint64>();
-}
-
-TOptional<float> UAnimNextComponent::GetVariableFloat(const FName Name) const
-{
-	const int32* IndexPtr = PublicVariablesProxyMap.Find(Name);
-	if (IndexPtr == nullptr)
-	{
-		return TOptional<float>();
-	}
-
-	FRWScopeLock Lock(PublicVariablesLock, SLT_ReadOnly);
-
-	TConstArrayView<FPropertyBagPropertyDesc> ProxyDescs = PublicVariablesProxy.Data.GetPropertyBagStruct()->GetPropertyDescs();
-	const FPropertyBagPropertyDesc& PropertyDesc = ProxyDescs[*IndexPtr];
-
-	TValueOrError<float, EPropertyBagResult> Result = PublicVariablesProxy.Data.GetValueFloat(PropertyDesc);
-	return Result.HasValue() ? TOptional<float>(Result.GetValue()) : TOptional<float>();
-}
-
-TOptional<double> UAnimNextComponent::GetVariableDouble(const FName Name) const
-{
-	const int32* IndexPtr = PublicVariablesProxyMap.Find(Name);
-	if (IndexPtr == nullptr)
-	{
-		return TOptional<double>();
-	}
-
-	FRWScopeLock Lock(PublicVariablesLock, SLT_ReadOnly);
-
-	TConstArrayView<FPropertyBagPropertyDesc> ProxyDescs = PublicVariablesProxy.Data.GetPropertyBagStruct()->GetPropertyDescs();
-	const FPropertyBagPropertyDesc& PropertyDesc = ProxyDescs[*IndexPtr];
-
-	TValueOrError<double, EPropertyBagResult> Result = PublicVariablesProxy.Data.GetValueDouble(PropertyDesc);
-	return Result.HasValue() ? TOptional<double>(Result.GetValue()) : TOptional<double>();
-}
-
-TOptional<FName> UAnimNextComponent::GetVariableName(const FName Name) const
-{
-	const int32* IndexPtr = PublicVariablesProxyMap.Find(Name);
-	if (IndexPtr == nullptr)
-	{
-		return TOptional<FName>();
-	}
-
-	FRWScopeLock Lock(PublicVariablesLock, SLT_ReadOnly);
-
-	TConstArrayView<FPropertyBagPropertyDesc> ProxyDescs = PublicVariablesProxy.Data.GetPropertyBagStruct()->GetPropertyDescs();
-	const FPropertyBagPropertyDesc& PropertyDesc = ProxyDescs[*IndexPtr];
-
-	TValueOrError<FName, EPropertyBagResult> Result = PublicVariablesProxy.Data.GetValueName(PropertyDesc);
-	return Result.HasValue() ? TOptional<FName>(Result.GetValue()) : TOptional<FName>();
-}
-
-TOptional<FString> UAnimNextComponent::GetVariableString(const FName Name) const
-{
-	const int32* IndexPtr = PublicVariablesProxyMap.Find(Name);
-	if (IndexPtr == nullptr)
-	{
-		return TOptional<FString>();
-	}
-
-	FRWScopeLock Lock(PublicVariablesLock, SLT_ReadOnly);
-
-	TConstArrayView<FPropertyBagPropertyDesc> ProxyDescs = PublicVariablesProxy.Data.GetPropertyBagStruct()->GetPropertyDescs();
-	const FPropertyBagPropertyDesc& PropertyDesc = ProxyDescs[*IndexPtr];
-
-	TValueOrError<FString, EPropertyBagResult> Result = PublicVariablesProxy.Data.GetValueString(PropertyDesc);
-	return Result.HasValue() ? TOptional<FString>(Result.GetValue()) : TOptional<FString>();
-}
-
-TOptional<uint8> UAnimNextComponent::GetVariableEnum(const FName Name, const UEnum* RequestedEnum) const
-{
-	const int32* IndexPtr = PublicVariablesProxyMap.Find(Name);
-	if (IndexPtr == nullptr)
+	switch (Result)
 	{
-		return TOptional<uint8>();
-	}
-
-	FRWScopeLock Lock(PublicVariablesLock, SLT_ReadOnly);
-
-	TConstArrayView<FPropertyBagPropertyDesc> ProxyDescs = PublicVariablesProxy.Data.GetPropertyBagStruct()->GetPropertyDescs();
-	const FPropertyBagPropertyDesc& PropertyDesc = ProxyDescs[*IndexPtr];
-
-	TValueOrError<uint8, EPropertyBagResult> Result = PublicVariablesProxy.Data.GetValueEnum(PropertyDesc, RequestedEnum);
-	return Result.HasValue() ? TOptional<uint8>(Result.GetValue()) : TOptional<uint8>();
-}
-
-TOptional<FConstStructView> UAnimNextComponent::GetVariableStruct(const FName Name, const UScriptStruct* RequestedStruct) const
-{
-	const int32* IndexPtr = PublicVariablesProxyMap.Find(Name);
-	if (IndexPtr == nullptr)
-	{
-		return TOptional<FConstStructView>();
-	}
-
-	FRWScopeLock Lock(PublicVariablesLock, SLT_ReadOnly);
-
-	TConstArrayView<FPropertyBagPropertyDesc> ProxyDescs = PublicVariablesProxy.Data.GetPropertyBagStruct()->GetPropertyDescs();
-	const FPropertyBagPropertyDesc& PropertyDesc = ProxyDescs[*IndexPtr];
-
-	TValueOrError<FStructView, EPropertyBagResult> Result = PublicVariablesProxy.Data.GetValueStruct(PropertyDesc, RequestedStruct);
-	return Result.HasValue() ? TOptional<FConstStructView>(Result.GetValue()) : TOptional<FConstStructView>();
-}
-
-TOptional<const UObject*> UAnimNextComponent::GetVariableObject(const FName Name, const UClass* RequestedClass) const
-{
-	const int32* IndexPtr = PublicVariablesProxyMap.Find(Name);
-	if (IndexPtr == nullptr)
-	{
-		return TOptional<const UObject*>();
-	}
-
-	FRWScopeLock Lock(PublicVariablesLock, SLT_ReadOnly);
-
-	TConstArrayView<FPropertyBagPropertyDesc> ProxyDescs = PublicVariablesProxy.Data.GetPropertyBagStruct()->GetPropertyDescs();
-	const FPropertyBagPropertyDesc& PropertyDesc = ProxyDescs[*IndexPtr];
-
-	TValueOrError<UObject*, EPropertyBagResult> Result = PublicVariablesProxy.Data.GetValueObject(PropertyDesc, RequestedClass);
-	return Result.HasValue() ? TOptional<const UObject*>(Result.GetValue()) : TOptional<const UObject*>();
-}
-
-TOptional<const UClass*> UAnimNextComponent::GetVariableClass(const FName Name) const
-{
-	const int32* IndexPtr = PublicVariablesProxyMap.Find(Name);
-	if (IndexPtr == nullptr)
-	{
-		return TOptional<const UClass*>();
-	}
-
-	FRWScopeLock Lock(PublicVariablesLock, SLT_ReadOnly);
-
-	TConstArrayView<FPropertyBagPropertyDesc> ProxyDescs = PublicVariablesProxy.Data.GetPropertyBagStruct()->GetPropertyDescs();
-	const FPropertyBagPropertyDesc& PropertyDesc = ProxyDescs[*IndexPtr];
-
-	TValueOrError<UClass*, EPropertyBagResult> Result = PublicVariablesProxy.Data.GetValueClass(PropertyDesc);
-	return Result.HasValue() ? TOptional<const UClass*>(Result.GetValue()) : TOptional<const UClass*>();
-}
-
-TOptional<FSoftObjectPath> UAnimNextComponent::GetVariableSoftPath(const FName Name) const
-{
-	const int32* IndexPtr = PublicVariablesProxyMap.Find(Name);
-	if (IndexPtr == nullptr)
-	{
-		return TOptional<FSoftObjectPath>();
-	}
-
-	FRWScopeLock Lock(PublicVariablesLock, SLT_ReadOnly);
-
-	TConstArrayView<FPropertyBagPropertyDesc> ProxyDescs = PublicVariablesProxy.Data.GetPropertyBagStruct()->GetPropertyDescs();
-	const FPropertyBagPropertyDesc& PropertyDesc = ProxyDescs[*IndexPtr];
-
-	TValueOrError<FSoftObjectPath, EPropertyBagResult> Result = PublicVariablesProxy.Data.GetValueSoftPath(PropertyDesc);
-	return Result.HasValue() ? TOptional<FSoftObjectPath>(Result.GetValue()) : TOptional<FSoftObjectPath>();
-}
-
-bool UAnimNextComponent::SetVariableBool(const FName Name, const bool bInValue)
-{
-	const int32* IndexPtr = PublicVariablesProxyMap.Find(Name);
-	if (IndexPtr == nullptr)
-	{
-		return false;
-	}
-
-	FRWScopeLock Lock(PublicVariablesLock, SLT_Write);
-
-	TConstArrayView<FPropertyBagPropertyDesc> ProxyDescs = PublicVariablesProxy.Data.GetPropertyBagStruct()->GetPropertyDescs();
-	const FPropertyBagPropertyDesc& PropertyDesc = ProxyDescs[*IndexPtr];
-
-	EPropertyBagResult Result = PublicVariablesProxy.Data.SetValueBool(PropertyDesc, bInValue);
-	if (Result == EPropertyBagResult::Success)
-	{
-		PublicVariablesProxy.DirtyFlags[*IndexPtr] = true;
-		PublicVariablesProxy.bIsDirty = true;
-	}
-
-	return Result == EPropertyBagResult::Success;
-}
-
-bool UAnimNextComponent::SetVariableByte(const FName Name, const uint8 InValue)
-{
-	const int32* IndexPtr = PublicVariablesProxyMap.Find(Name);
-	if (IndexPtr == nullptr)
-	{
-		return false;
-	}
-
-	FRWScopeLock Lock(PublicVariablesLock, SLT_Write);
-
-	TConstArrayView<FPropertyBagPropertyDesc> ProxyDescs = PublicVariablesProxy.Data.GetPropertyBagStruct()->GetPropertyDescs();
-	const FPropertyBagPropertyDesc& PropertyDesc = ProxyDescs[*IndexPtr];
-
-	EPropertyBagResult Result = PublicVariablesProxy.Data.SetValueByte(PropertyDesc, InValue);
-	if (Result == EPropertyBagResult::Success)
-	{
-		PublicVariablesProxy.DirtyFlags[*IndexPtr] = true;
-		PublicVariablesProxy.bIsDirty = true;
-	}
-
-	return Result == EPropertyBagResult::Success;
-}
-
-bool UAnimNextComponent::SetVariableInt32(const FName Name, const int32 InValue)
-{
-	const int32* IndexPtr = PublicVariablesProxyMap.Find(Name);
-	if (IndexPtr == nullptr)
-	{
-		return false;
-	}
-
-	FRWScopeLock Lock(PublicVariablesLock, SLT_Write);
-
-	TConstArrayView<FPropertyBagPropertyDesc> ProxyDescs = PublicVariablesProxy.Data.GetPropertyBagStruct()->GetPropertyDescs();
-	const FPropertyBagPropertyDesc& PropertyDesc = ProxyDescs[*IndexPtr];
-
-	EPropertyBagResult Result = PublicVariablesProxy.Data.SetValueInt32(PropertyDesc, InValue);
-	if (Result == EPropertyBagResult::Success)
-	{
-		PublicVariablesProxy.DirtyFlags[*IndexPtr] = true;
-		PublicVariablesProxy.bIsDirty = true;
-	}
-
-	return Result == EPropertyBagResult::Success;
-}
-
-bool UAnimNextComponent::SetVariableUInt32(const FName Name, const uint32 InValue)
-{
-	const int32* IndexPtr = PublicVariablesProxyMap.Find(Name);
-	if (IndexPtr == nullptr)
-	{
-		return false;
-	}
-
-	FRWScopeLock Lock(PublicVariablesLock, SLT_Write);
-
-	TConstArrayView<FPropertyBagPropertyDesc> ProxyDescs = PublicVariablesProxy.Data.GetPropertyBagStruct()->GetPropertyDescs();
-	const FPropertyBagPropertyDesc& PropertyDesc = ProxyDescs[*IndexPtr];
-
-	EPropertyBagResult Result = PublicVariablesProxy.Data.SetValueUInt32(PropertyDesc, InValue);
-	if (Result == EPropertyBagResult::Success)
-	{
-		PublicVariablesProxy.DirtyFlags[*IndexPtr] = true;
-		PublicVariablesProxy.bIsDirty = true;
-	}
-
-	return Result == EPropertyBagResult::Success;
-}
-
-bool UAnimNextComponent::SetVariableInt64(const FName Name, const int64 InValue)
-{
-	const int32* IndexPtr = PublicVariablesProxyMap.Find(Name);
-	if (IndexPtr == nullptr)
-	{
-		return false;
-	}
-
-	FRWScopeLock Lock(PublicVariablesLock, SLT_Write);
-
-	TConstArrayView<FPropertyBagPropertyDesc> ProxyDescs = PublicVariablesProxy.Data.GetPropertyBagStruct()->GetPropertyDescs();
-	const FPropertyBagPropertyDesc& PropertyDesc = ProxyDescs[*IndexPtr];
-
-	EPropertyBagResult Result = PublicVariablesProxy.Data.SetValueInt64(PropertyDesc, InValue);
-	if (Result == EPropertyBagResult::Success)
-	{
-		PublicVariablesProxy.DirtyFlags[*IndexPtr] = true;
-		PublicVariablesProxy.bIsDirty = true;
-	}
-
-	return Result == EPropertyBagResult::Success;
-}
-
-bool UAnimNextComponent::SetVariableUInt64(const FName Name, const uint64 InValue)
-{
-	const int32* IndexPtr = PublicVariablesProxyMap.Find(Name);
-	if (IndexPtr == nullptr)
-	{
-		return false;
-	}
-
-	FRWScopeLock Lock(PublicVariablesLock, SLT_Write);
-
-	TConstArrayView<FPropertyBagPropertyDesc> ProxyDescs = PublicVariablesProxy.Data.GetPropertyBagStruct()->GetPropertyDescs();
-	const FPropertyBagPropertyDesc& PropertyDesc = ProxyDescs[*IndexPtr];
-
-	EPropertyBagResult Result = PublicVariablesProxy.Data.SetValueUInt64(PropertyDesc, InValue);
-	if (Result == EPropertyBagResult::Success)
-	{
-		PublicVariablesProxy.DirtyFlags[*IndexPtr] = true;
-		PublicVariablesProxy.bIsDirty = true;
-	}
-
-	return Result == EPropertyBagResult::Success;
-}
-
-bool UAnimNextComponent::SetVariableFloat(const FName Name, const float InValue)
-{
-	const int32* IndexPtr = PublicVariablesProxyMap.Find(Name);
-	if (IndexPtr == nullptr)
-	{
-		return false;
-	}
-
-	FRWScopeLock Lock(PublicVariablesLock, SLT_Write);
-
-	TConstArrayView<FPropertyBagPropertyDesc> ProxyDescs = PublicVariablesProxy.Data.GetPropertyBagStruct()->GetPropertyDescs();
-	const FPropertyBagPropertyDesc& PropertyDesc = ProxyDescs[*IndexPtr];
-
-	EPropertyBagResult Result = PublicVariablesProxy.Data.SetValueFloat(PropertyDesc, InValue);
-	if (Result == EPropertyBagResult::Success)
-	{
-		PublicVariablesProxy.DirtyFlags[*IndexPtr] = true;
-		PublicVariablesProxy.bIsDirty = true;
-	}
-
-	return Result == EPropertyBagResult::Success;
-}
-
-bool UAnimNextComponent::SetVariableDouble(const FName Name, const double InValue)
-{
-	const int32* IndexPtr = PublicVariablesProxyMap.Find(Name);
-	if (IndexPtr == nullptr)
-	{
-		return false;
-	}
-
-	FRWScopeLock Lock(PublicVariablesLock, SLT_Write);
-
-	TConstArrayView<FPropertyBagPropertyDesc> ProxyDescs = PublicVariablesProxy.Data.GetPropertyBagStruct()->GetPropertyDescs();
-	const FPropertyBagPropertyDesc& PropertyDesc = ProxyDescs[*IndexPtr];
-
-	EPropertyBagResult Result = PublicVariablesProxy.Data.SetValueDouble(PropertyDesc, InValue);
-	if (Result == EPropertyBagResult::Success)
-	{
-		PublicVariablesProxy.DirtyFlags[*IndexPtr] = true;
-		PublicVariablesProxy.bIsDirty = true;
-	}
-
-	return Result == EPropertyBagResult::Success;
-}
-
-bool UAnimNextComponent::SetVariableName(const FName Name, const FName InValue)
-{
-	const int32* IndexPtr = PublicVariablesProxyMap.Find(Name);
-	if (IndexPtr == nullptr)
-	{
-		return false;
-	}
-
-	FRWScopeLock Lock(PublicVariablesLock, SLT_Write);
-
-	TConstArrayView<FPropertyBagPropertyDesc> ProxyDescs = PublicVariablesProxy.Data.GetPropertyBagStruct()->GetPropertyDescs();
-	const FPropertyBagPropertyDesc& PropertyDesc = ProxyDescs[*IndexPtr];
-
-	EPropertyBagResult Result = PublicVariablesProxy.Data.SetValueName(PropertyDesc, InValue);
-	if (Result == EPropertyBagResult::Success)
-	{
-		PublicVariablesProxy.DirtyFlags[*IndexPtr] = true;
-		PublicVariablesProxy.bIsDirty = true;
-	}
-
-	return Result == EPropertyBagResult::Success;
-}
-
-bool UAnimNextComponent::SetVariableString(const FName Name, const FString& InValue)
-{
-	const int32* IndexPtr = PublicVariablesProxyMap.Find(Name);
-	if (IndexPtr == nullptr)
-	{
-		return false;
-	}
+	case EPropertyBagResult::Success:
+		break;
+	case EPropertyBagResult::TypeMismatch:
+		{
+			FBlueprintExceptionInfo ExceptionInfo(
+				EBlueprintExceptionType::NonFatalError,
+				FText::Format(NSLOCTEXT("UAFComponent", "UAFComponent_SetVariableTypeMismatch", "Incompatible type supplied for variable '{0}'"), FText::FromName(Name))
+			);
 
-	FRWScopeLock Lock(PublicVariablesLock, SLT_Write);
+			FBlueprintCoreDelegates::ThrowScriptException(P_THIS, Stack, ExceptionInfo);
+			break;
+		}
+	case EPropertyBagResult::PropertyNotFound:
+		{
+			FBlueprintExceptionInfo ExceptionInfo(
+				EBlueprintExceptionType::NonFatalError,
+				FText::Format(NSLOCTEXT("UAFComponent", "UAFComponent_SetVariableNotFoundWarning", "Unknown variable name '{0}' supplied to Set Variable"), FText::FromName(Name))
+			);
 
-	TConstArrayView<FPropertyBagPropertyDesc> ProxyDescs = PublicVariablesProxy.Data.GetPropertyBagStruct()->GetPropertyDescs();
-	const FPropertyBagPropertyDesc& PropertyDesc = ProxyDescs[*IndexPtr];
+			FBlueprintCoreDelegates::ThrowScriptException(P_THIS, Stack, ExceptionInfo);
+			break;
+		}
+	case EPropertyBagResult::OutOfBounds:
+	case EPropertyBagResult::DuplicatedValue:
+		{
+			FBlueprintExceptionInfo ExceptionInfo(
+				EBlueprintExceptionType::NonFatalError,
+				FText::Format(NSLOCTEXT("UAFComponent", "UAFComponent_SetVariableUnknownWarning", "Unexpected internal error when calling Set Variable for '{0}'"), FText::FromName(Name))
+			);
 
-	EPropertyBagResult Result = PublicVariablesProxy.Data.SetValueString(PropertyDesc, InValue);
-	if (Result == EPropertyBagResult::Success)
-	{
-		PublicVariablesProxy.DirtyFlags[*IndexPtr] = true;
-		PublicVariablesProxy.bIsDirty = true;
+			FBlueprintCoreDelegates::ThrowScriptException(P_THIS, Stack, ExceptionInfo);
+			break;
+		}
 	}
-
-	return Result == EPropertyBagResult::Success;
 }
 
-bool UAnimNextComponent::SetVariableEnum(const FName Name, const uint8 InValue, const UEnum* Enum)
+void UAnimNextComponent::BlueprintSetVariableReference(const FAnimNextVariableReference& Variable, const int32& Value)
 {
-	const int32* IndexPtr = PublicVariablesProxyMap.Find(Name);
-	if (IndexPtr == nullptr)
-	{
-		return false;
-	}
-
-	FRWScopeLock Lock(PublicVariablesLock, SLT_Write);
-
-	TConstArrayView<FPropertyBagPropertyDesc> ProxyDescs = PublicVariablesProxy.Data.GetPropertyBagStruct()->GetPropertyDescs();
-	const FPropertyBagPropertyDesc& PropertyDesc = ProxyDescs[*IndexPtr];
-
-	EPropertyBagResult Result = PublicVariablesProxy.Data.SetValueEnum(PropertyDesc, InValue, Enum);
-	if (Result == EPropertyBagResult::Success)
-	{
-		PublicVariablesProxy.DirtyFlags[*IndexPtr] = true;
-		PublicVariablesProxy.bIsDirty = true;
-	}
-
-	return Result == EPropertyBagResult::Success;
+	checkNoEntry();
 }
 
-bool UAnimNextComponent::SetVariableStruct(const FName Name, FConstStructView InValue)
+DEFINE_FUNCTION(UAnimNextComponent::execBlueprintSetVariableReference)
 {
-	const int32* IndexPtr = PublicVariablesProxyMap.Find(Name);
-	if (IndexPtr == nullptr)
-	{
-		return false;
-	}
-
-	FRWScopeLock Lock(PublicVariablesLock, SLT_Write);
+	using namespace UE::AnimNext;
 
-	TConstArrayView<FPropertyBagPropertyDesc> ProxyDescs = PublicVariablesProxy.Data.GetPropertyBagStruct()->GetPropertyDescs();
-	const FPropertyBagPropertyDesc& PropertyDesc = ProxyDescs[*IndexPtr];
+	P_GET_STRUCT_REF(FAnimNextVariableReference, Variable);
 
-	EPropertyBagResult Result = PublicVariablesProxy.Data.SetValueStruct(PropertyDesc, InValue);
-	if (Result == EPropertyBagResult::Success)
+	const FProperty* VariableProperty = Variable.ResolveProperty();
+	if (VariableProperty == nullptr)
 	{
-		PublicVariablesProxy.DirtyFlags[*IndexPtr] = true;
-		PublicVariablesProxy.bIsDirty = true;
-	}
+		P_FINISH;
 
-	return Result == EPropertyBagResult::Success;
-}
+		FBlueprintExceptionInfo ExceptionInfo(
+			EBlueprintExceptionType::AbortExecution,
+			FText::Format(NSLOCTEXT("UAFComponent", "UAFComponent_SetVariablePropertyError", "Failed to resolve the property {0} of variable reference"), FText::FromName(Variable.GetName()))
+		);
 
-bool UAnimNextComponent::SetVariableStructRef(const FName Name, TFunctionRef<void(FStructView)> InStructRefSetter, const UScriptStruct* RequestedStruct)
-{
-	const int32* IndexPtr = PublicVariablesProxyMap.Find(Name);
-	if (IndexPtr == nullptr)
-	{
-		return false;
+		FBlueprintCoreDelegates::ThrowScriptException(P_THIS, Stack, ExceptionInfo);
+		return;
 	}
 
-	FRWScopeLock Lock(PublicVariablesLock, SLT_Write);
-
-	TConstArrayView<FPropertyBagPropertyDesc> ProxyDescs = PublicVariablesProxy.Data.GetPropertyBagStruct()->GetPropertyDescs();
-	const FPropertyBagPropertyDesc& PropertyDesc = ProxyDescs[*IndexPtr];
-
-	TValueOrError<FStructView, EPropertyBagResult> Result = PublicVariablesProxy.Data.GetValueStruct(PropertyDesc, RequestedStruct);
-	if (Result.HasValue())
+	if (Variable.GetName() == NAME_None || Variable.GetObject() == nullptr)
 	{
-		InStructRefSetter(Result.GetValue());
-		PublicVariablesProxy.DirtyFlags[*IndexPtr] = true;
-		PublicVariablesProxy.bIsDirty = true;
-		return true;
-	}
+		P_FINISH;
 
-	return false;
-}
+		FBlueprintExceptionInfo ExceptionInfo(
+			EBlueprintExceptionType::NonFatalError,
+			NSLOCTEXT("UAFComponent", "UAFComponent_SetVariableInvalidWarning", "Invalid variable supplied to Set Variable")
+		);
 
-bool UAnimNextComponent::SetVariableObject(const FName Name, UObject* InValue)
-{
-	const int32* IndexPtr = PublicVariablesProxyMap.Find(Name);
-	if (IndexPtr == nullptr)
-	{
-		return false;
+		FBlueprintCoreDelegates::ThrowScriptException(P_THIS, Stack, ExceptionInfo);
+		return;
 	}
 
-	FRWScopeLock Lock(PublicVariablesLock, SLT_Write);
-
-	TConstArrayView<FPropertyBagPropertyDesc> ProxyDescs = PublicVariablesProxy.Data.GetPropertyBagStruct()->GetPropertyDescs();
-	const FPropertyBagPropertyDesc& PropertyDesc = ProxyDescs[*IndexPtr];
+	// Read wildcard Value input.
+	Stack.MostRecentPropertyAddress = nullptr;
+	Stack.MostRecentPropertyContainer = nullptr;
 
-	EPropertyBagResult Result = PublicVariablesProxy.Data.SetValueObject(PropertyDesc, InValue);
-	if (Result == EPropertyBagResult::Success)
+	EPropertyBagResult Result = EPropertyBagResult::Success;
 	{
-		PublicVariablesProxy.DirtyFlags[*IndexPtr] = true;
-		PublicVariablesProxy.bIsDirty = true;
-	}
+		uint8* TempStorage = (uint8*)FMemory_Alloca_Aligned(VariableProperty->GetElementSize(), VariableProperty->GetMinAlignment());
+		VariableProperty->InitializeValue(TempStorage);
 
-	return Result == EPropertyBagResult::Success;
-}
+		Stack.StepCompiledIn(TempStorage, VariableProperty->GetClass());
 
-bool UAnimNextComponent::SetVariableObjectRef(const FName Name, TFunctionRef<void(UObject*)> InObjectRefSetter, const UClass* RequestedClass)
-{
-	const int32* IndexPtr = PublicVariablesProxyMap.Find(Name);
-	if (IndexPtr == nullptr)
-	{
-		return false;
-	}
+		P_FINISH;
+		
+		P_NATIVE_BEGIN;
 
-	FRWScopeLock Lock(PublicVariablesLock, SLT_Write);
+		{
+			FAnimNextParamType Type = FAnimNextParamType::FromProperty(VariableProperty);
+			Result = P_THIS->SetVariableInternal(Variable, Type, TConstArrayView<uint8>(TempStorage, VariableProperty->GetElementSize()));
+		}
 
-	TConstArrayView<FPropertyBagPropertyDesc> ProxyDescs = PublicVariablesProxy.Data.GetPropertyBagStruct()->GetPropertyDescs();
-	const FPropertyBagPropertyDesc& PropertyDesc = ProxyDescs[*IndexPtr];
+		P_NATIVE_END;
 
-	TValueOrError<UObject*, EPropertyBagResult> Result = PublicVariablesProxy.Data.GetValueObject(PropertyDesc, RequestedClass);
-	if (Result.HasValue())
-	{
-		InObjectRefSetter(Result.GetValue());
-		PublicVariablesProxy.DirtyFlags[*IndexPtr] = true;
-		PublicVariablesProxy.bIsDirty = true;
-		return true;
+		VariableProperty->DestroyValue(TempStorage);
 	}
 
-	return false;
-}
-
-bool UAnimNextComponent::SetVariableClass(const FName Name, UClass* InValue)
-{
-	const int32* IndexPtr = PublicVariablesProxyMap.Find(Name);
-	if (IndexPtr == nullptr)
+	switch (Result)
 	{
-		return false;
-	}
+	case EPropertyBagResult::Success:
+		break;
+	case EPropertyBagResult::TypeMismatch:
+		{
+			FBlueprintExceptionInfo ExceptionInfo(
+				EBlueprintExceptionType::NonFatalError,
+				FText::Format(NSLOCTEXT("UAFComponent", "UAFComponent_SetVariableTypeMismatch", "Incompatible type supplied for variable '{0}'"), FText::FromName(Variable.GetName()))
+			);
 
-	FRWScopeLock Lock(PublicVariablesLock, SLT_Write);
+			FBlueprintCoreDelegates::ThrowScriptException(P_THIS, Stack, ExceptionInfo);
+			break;
+		}
+	case EPropertyBagResult::PropertyNotFound:
+		{
+			FBlueprintExceptionInfo ExceptionInfo(
+				EBlueprintExceptionType::NonFatalError,
+				FText::Format(NSLOCTEXT("UAFComponent", "UAFComponent_SetVariableNotFoundWarning", "Unknown variable name '{0}' supplied to Set Variable"), FText::FromName(Variable.GetName()))
+			);
 
-	TConstArrayView<FPropertyBagPropertyDesc> ProxyDescs = PublicVariablesProxy.Data.GetPropertyBagStruct()->GetPropertyDescs();
-	const FPropertyBagPropertyDesc& PropertyDesc = ProxyDescs[*IndexPtr];
+			FBlueprintCoreDelegates::ThrowScriptException(P_THIS, Stack, ExceptionInfo);
+			break;
+		}
+	case EPropertyBagResult::OutOfBounds:
+	case EPropertyBagResult::DuplicatedValue:
+		{
+			FBlueprintExceptionInfo ExceptionInfo(
+				EBlueprintExceptionType::NonFatalError,
+				FText::Format(NSLOCTEXT("UAFComponent", "UAFComponent_SetVariableUnknownWarning", "Unexpected internal error when calling Set Variable for '{0}'"), FText::FromName(Variable.GetName()))
+			);
 
-	EPropertyBagResult Result = PublicVariablesProxy.Data.SetValueClass(PropertyDesc, InValue);
-	if (Result == EPropertyBagResult::Success)
-	{
-		PublicVariablesProxy.DirtyFlags[*IndexPtr] = true;
-		PublicVariablesProxy.bIsDirty = true;
+			FBlueprintCoreDelegates::ThrowScriptException(P_THIS, Stack, ExceptionInfo);
+			break;
+		}
 	}
-
-	return Result == EPropertyBagResult::Success;
 }
 
-bool UAnimNextComponent::SetVariableSoftPath(const FName Name, const FSoftObjectPath& InValue)
+EPropertyBagResult UAnimNextComponent::SetVariableInternal(const FAnimNextVariableReference& InVariable, const FAnimNextParamType& InType, TConstArrayView<uint8> InNewValue)
 {
-	const int32* IndexPtr = PublicVariablesProxyMap.Find(Name);
-	if (IndexPtr == nullptr)
-	{
-		return false;
-	}
-
-	FRWScopeLock Lock(PublicVariablesLock, SLT_Write);
-
-	TConstArrayView<FPropertyBagPropertyDesc> ProxyDescs = PublicVariablesProxy.Data.GetPropertyBagStruct()->GetPropertyDescs();
-	const FPropertyBagPropertyDesc& PropertyDesc = ProxyDescs[*IndexPtr];
-
-	EPropertyBagResult Result = PublicVariablesProxy.Data.SetValueSoftPath(PropertyDesc, InValue);
-	if (Result == EPropertyBagResult::Success)
+	if (Subsystem)
 	{
-		PublicVariablesProxy.DirtyFlags[*IndexPtr] = true;
-		PublicVariablesProxy.bIsDirty = true;
+		return Subsystem->SetVariable(this, InVariable, InType, InNewValue);
 	}
-
-	return Result == EPropertyBagResult::Success;
+	return EPropertyBagResult::PropertyNotFound;
 }
 
-bool UAnimNextComponent::SetVariableArrayRef(const FName Name, TFunctionRef<void(FPropertyBagArrayRef&)> InArrayRefSetter)
+EPropertyBagResult UAnimNextComponent::WriteVariableInternal(const FAnimNextVariableReference& InVariable, const FAnimNextParamType& InType, TFunctionRef<void(TArrayView<uint8>)> InFunction)
 {
-	const int32* IndexPtr = PublicVariablesProxyMap.Find(Name);
-	if (IndexPtr == nullptr)
-	{
-		return false;
-	}
-
-	FRWScopeLock Lock(PublicVariablesLock, SLT_Write);
-
-	TConstArrayView<FPropertyBagPropertyDesc> ProxyDescs = PublicVariablesProxy.Data.GetPropertyBagStruct()->GetPropertyDescs();
-	const FPropertyBagPropertyDesc& PropertyDesc = ProxyDescs[*IndexPtr];
-
-	TValueOrError<FPropertyBagArrayRef, EPropertyBagResult> Result = PublicVariablesProxy.Data.GetMutableArrayRef(PropertyDesc);
-	if (Result.HasValue())
+	if (Subsystem)
 	{
-		InArrayRefSetter(Result.GetValue());
-		PublicVariablesProxy.DirtyFlags[*IndexPtr] = true;
-		PublicVariablesProxy.bIsDirty = true;
-		return true;
+		return Subsystem->WriteVariable(this, InVariable, InType, InFunction);
 	}
-
-	return false;
+	return EPropertyBagResult::PropertyNotFound;
 }
 
 bool UAnimNextComponent::IsEnabled() const
diff --git a/Engine/Plugins/Experimental/AnimNext/Source/AnimNext/Private/Component/AnimNextComponentWorldSubsystem.cpp b/Engine/Plugins/Experimental/AnimNext/Source/AnimNext/Private/Component/AnimNextComponentWorldSubsystem.cpp
--- a/Engine/Plugins/Experimental/AnimNext/Source/AnimNext/Private/Component/AnimNextComponentWorldSubsystem.cpp
+++ b/Engine/Plugins/Experimental/AnimNext/Source/AnimNext/Private/Component/AnimNextComponentWorldSubsystem.cpp
@@ -14,7 +14,7 @@ void UAnimNextComponentWorldSubsystem::Register(UAnimNextComponent* InComponent)
 	using namespace UE::AnimNext;
 	check(IsInGameThread());
 	check(InComponent);
-	RegisterHandle(InComponent->ModuleHandle, InComponent->Module, InComponent, InComponent, InComponent->InitMethod);
+	RegisterHandle(InComponent->ModuleHandle, InComponent->Module, InComponent, InComponent->InitMethod);
 }
 
 void UAnimNextComponentWorldSubsystem::Unregister(UAnimNextComponent* InComponent)
@@ -107,20 +107,16 @@ void UAnimNextComponentWorldSubsystem::RemoveModuleEventDependency(UAnimNextComp
 	RemoveModuleEventDependencyHandle(InComponent->ModuleHandle, InEventName, OtherComponent->ModuleHandle, OtherEventName, InDependency);
 }
 
-#if WITH_EDITOR
-
-void UAnimNextComponentWorldSubsystem::OnModuleCompiled(UAnimNextModule* InModule)
+EPropertyBagResult UAnimNextComponentWorldSubsystem::SetVariable(UAnimNextComponent* InComponent, const FAnimNextVariableReference& InVariable, const FAnimNextParamType& InType, TConstArrayView<uint8> InData)
 {
-	Super::OnModuleCompiled(InModule);
-
-	for(FAnimNextModuleInstance& Instance : Instances)
-	{
-		if(Instance.GetModule() == InModule)
-		{
-			UAnimNextComponent* AnimNextComponent = CastChecked<UAnimNextComponent>(Instance.GetObject());
-			AnimNextComponent->OnModuleCompiled();
-		}
-	}
+	using namespace UE::AnimNext;
+	check(InComponent);
+	return SetVariableHandle(InComponent->ModuleHandle, InVariable, InType, InData);
 }
 
-#endif
\ No newline at end of file
+EPropertyBagResult UAnimNextComponentWorldSubsystem::WriteVariable(UAnimNextComponent* InComponent, const FAnimNextVariableReference& InVariable, const FAnimNextParamType& InType, TFunctionRef<void(TArrayView<uint8>)> InFunction)
+{
+	using namespace UE::AnimNext;
+	check(InComponent);
+	return WriteVariableHandle(InComponent->ModuleHandle, InVariable, InType, InFunction);
+}
\ No newline at end of file
diff --git a/Engine/Plugins/Experimental/AnimNext/Source/AnimNext/Private/Component/AnimNextWorldSubsystem.cpp b/Engine/Plugins/Experimental/AnimNext/Source/AnimNext/Private/Component/AnimNextWorldSubsystem.cpp
--- a/Engine/Plugins/Experimental/AnimNext/Source/AnimNext/Private/Component/AnimNextWorldSubsystem.cpp
+++ b/Engine/Plugins/Experimental/AnimNext/Source/AnimNext/Private/Component/AnimNextWorldSubsystem.cpp
@@ -8,7 +8,7 @@
 #include "Module/ModuleHandle.h"
 #include "Module/ModuleTaskContext.h"
 #include "AnimNextDebugDraw.h"
-#include "Module/ModuleGuard.h"
+#include "Variables/AnimNextVariableReference.h"
 
 UAnimNextWorldSubsystem::UAnimNextWorldSubsystem()
 {
@@ -48,7 +48,7 @@ void UAnimNextWorldSubsystem::BeginDestroy()
 			for (FAnimNextModuleInstance& Instance : Instances)
 			{
 				Instance.RemoveAllTickDependencies();
-				Instance.Uninitialize();
+				Instance.ReleaseComponents();
 			}
 			Instances = UE::AnimNext::TPool<FAnimNextModuleInstance>(); // Force instances destruction
 		}
@@ -119,8 +119,6 @@ void UAnimNextWorldSubsystem::FlushPendingActions()
 			case FModulePendingAction::EType::ReleaseHandle:
 				if (IsValidHandle(PendingAction.Handle))
 				{
-					FAnimNextModuleInstance& Instance = Instances.Get(PendingAction.Handle);
-					Instance.Uninitialize();
 					Instances.Release(PendingAction.Handle);
 				}
 				break;
@@ -140,44 +138,6 @@ void UAnimNextWorldSubsystem::FlushPendingActions()
 				}
 #endif
 				break;
-
-			case FModulePendingAction::EType::AddPrerequisite:
-				{					
-					FAnimNextModuleInstance* SubsequentInstance = Instances.TryGet(PendingAction.Handle);
-					FAnimNextModuleInstance* PrerequisiteInstance = Instances.TryGet(PendingAction.Payload.Get<UE::AnimNext::FModuleHandle>());
-					if(PrerequisiteInstance != nullptr && SubsequentInstance != nullptr)
-					{
-						SubsequentInstance->AddPrerequisite(*PrerequisiteInstance);
-					}
-					else
-					{
-						UE_LOGFMT(LogAnimation, Warning, "FlushPendingActions: Invalid Module handle(s) provided for AddPrerequisite action SubsequentHandle: %i PrerequisiteInstance: %i", PendingAction.Handle.IsValid(), PendingAction.Payload.Get<UE::AnimNext::FModuleHandle>().IsValid());
-					}
-				
-					break;
-				}
-			case FModulePendingAction::EType::RemovePrerequisite:
-				{					
-					FAnimNextModuleInstance* SubsequentInstance = Instances.TryGet(PendingAction.Handle);
-					FAnimNextModuleInstance* PrerequisiteInstance = Instances.TryGet(PendingAction.Payload.Get<UE::AnimNext::FModuleHandle>());
-					if(PrerequisiteInstance != nullptr && SubsequentInstance != nullptr)
-					{
-						if (SubsequentInstance->IsPrerequisite(*PrerequisiteInstance))
-						{							
-							SubsequentInstance->RemovePrerequisite(*PrerequisiteInstance);
-						}
-						else
-						{
-							UE_LOGFMT(LogAnimation, Warning, "FlushPendingActions: Trying to remove prerequisite ModuleInstance which isn't actually a prerequisite");
-						}
-					}
-					else
-					{
-						UE_LOGFMT(LogAnimation, Warning, "FlushPendingActions: Invalid Module handle(s) provided for RemovePrerequisite action SubsequentHandle: %i PrerequisiteInstance: %i", PendingAction.Handle.IsValid(), PendingAction.Payload.Get<UE::AnimNext::FModuleHandle>().IsValid());
-					}
-				
-					break;
-				}		
 			default:
 				checkNoEntry();
 				break;
@@ -188,13 +148,13 @@ void UAnimNextWorldSubsystem::FlushPendingActions()
 	}
 }
 
-void UAnimNextWorldSubsystem::RegisterHandle(UE::AnimNext::FModuleHandle& InOutHandle, UAnimNextModule* InModule, UObject* InObject, IAnimNextVariableProxyHost* InProxyHost, EAnimNextModuleInitMethod InitMethod)
+void UAnimNextWorldSubsystem::RegisterHandle(UE::AnimNext::FModuleHandle& InOutHandle, UAnimNextModule* InModule, UObject* InObject, EAnimNextModuleInitMethod InitMethod)
 {
 	using namespace UE::AnimNext;
 	check(IsInGameThread());
 	FWriteScopeLock InstancesLockScope(InstancesLock);
 
-	InOutHandle = Instances.Emplace(InModule, InObject, &Instances, InProxyHost, InitMethod);
+	InOutHandle = Instances.Emplace(InModule, InObject, &Instances, InitMethod);
 	FAnimNextModuleInstance& Instance = Instances.Get(InOutHandle);
 	Instance.Handle = InOutHandle;
 	Instance.Initialize();
@@ -218,12 +178,6 @@ void UAnimNextWorldSubsystem::UnregisterHandle(UE::AnimNext::FModuleHandle& InOu
 		// Remove all tick dependencies immediately, as once the handle has been invalidated there is no way for external systems to remove their dependencies
 		Instance.RemoveAllTickDependencies();
 
-		// Remove this module as a dependency on all other modules that depend upon it before the handle is invalidated
-		for (const UE::AnimNext::FModuleHandle& SubsequentHandle : Instance.SubsequentRefs)
-		{
-			RemoveModuleDependencyHandle(InOutHandle, SubsequentHandle);
-		}
-
 		PendingActions.Emplace(InOutHandle, FModulePendingAction::EType::ReleaseHandle);
 		InOutHandle.Reset();
 	}
@@ -312,30 +266,6 @@ void UAnimNextWorldSubsystem::QueueInputTraitEventHandle(UE::AnimNext::FModuleHa
 	ETaskRunLocation::Before);
 }
 
-void UAnimNextWorldSubsystem::AddModuleDependencyHandle(UE::AnimNext::FModuleHandle InPrerequisiteHandle, UE::AnimNext::FModuleHandle InSubsequentHandle)
-{
-	using namespace UE::AnimNext;
-	check(IsInGameThread());
-	if (!IsValidHandle(InPrerequisiteHandle) || !IsValidHandle(InSubsequentHandle))
-	{
-		return;
-	}
-
-	PendingActions.Emplace(InSubsequentHandle, FModulePendingAction::EType::AddPrerequisite, InPrerequisiteHandle);
-}
-
-void UAnimNextWorldSubsystem::RemoveModuleDependencyHandle(UE::AnimNext::FModuleHandle InPrerequisiteHandle, UE::AnimNext::FModuleHandle InSubsequentHandle)
-{
-	using namespace UE::AnimNext;
-	check(IsInGameThread());
-	if (!IsValidHandle(InPrerequisiteHandle) || !IsValidHandle(InSubsequentHandle))
-	{
-		return;
-	}
-
-	PendingActions.Emplace(InSubsequentHandle, FModulePendingAction::EType::RemovePrerequisite, InPrerequisiteHandle);
-}
-
 const FTickFunction* UAnimNextWorldSubsystem::FindTickFunctionHandle(UE::AnimNext::FModuleHandle InHandle, FName InEventName) const
 {
 	using namespace UE::AnimNext;
@@ -346,7 +276,7 @@ const FTickFunction* UAnimNextWorldSubsystem::FindTickFunctionHandle(UE::AnimNex
 		const FModuleEventTickFunction* TickFunction = Instance.FindTickFunctionByName(InEventName);
 		if(TickFunction == nullptr)
 		{
-			UE_LOGFMT(LogAnimation, Warning, "FindTickFunctionHandle: Could not find event '{EventName}' in module '{ModuleName}'", InEventName, Instance.GetDataInterfaceName());
+			UE_LOGFMT(LogAnimation, Warning, "FindTickFunctionHandle: Could not find event '{EventName}' in module '{ModuleName}'", InEventName, Instance.GetAssetName());
 			return nullptr;
 		}
 		
@@ -355,7 +285,7 @@ const FTickFunction* UAnimNextWorldSubsystem::FindTickFunctionHandle(UE::AnimNex
 			return TickFunction;
 		}
 
-		UE_LOGFMT(LogAnimation, Warning, "FindTickFunctionHandle: Event '{EventName}' in module '{ModuleName}' is not a bUserEvent, therefore cannot be exposed", InEventName, Instance.GetDataInterfaceName());
+		UE_LOGFMT(LogAnimation, Warning, "FindTickFunctionHandle: Event '{EventName}' in module '{ModuleName}' is not a bUserEvent, therefore cannot be exposed", InEventName, Instance.GetAssetName());
 	}
 	return nullptr;
 }
@@ -370,7 +300,7 @@ void UAnimNextWorldSubsystem::AddDependencyHandle(UE::AnimNext::FModuleHandle In
 		FModuleEventTickFunction* TickFunction = Instance.FindTickFunctionByName(InEventName);
 		if(TickFunction == nullptr)
 		{
-			UE_LOGFMT(LogAnimation, Warning, "AddDependencyHandle: Could not find event '{EventName}' in module '{ModuleName}'", InEventName, Instance.GetDataInterfaceName());
+			UE_LOGFMT(LogAnimation, Warning, "AddDependencyHandle: Could not find event '{EventName}' in module '{ModuleName}'", InEventName, Instance.GetAssetName());
 			return;
 		}
 
@@ -395,7 +325,7 @@ void UAnimNextWorldSubsystem::RemoveDependencyHandle(UE::AnimNext::FModuleHandle
 		FModuleEventTickFunction* TickFunction = Instance.FindTickFunctionByName(InEventName);
 		if(TickFunction == nullptr)
 		{
-			UE_LOGFMT(LogAnimation, Warning, "RemoveDependencyHandle: Could not find event '{EventName}' in module '{ModuleName}'", InEventName, Instance.GetDataInterfaceName());
+			UE_LOGFMT(LogAnimation, Warning, "RemoveDependencyHandle: Could not find event '{EventName}' in module '{ModuleName}'", InEventName, Instance.GetAssetName());
 			return;
 		}
 
@@ -421,7 +351,7 @@ void UAnimNextWorldSubsystem::AddModuleEventDependencyHandle(UE::AnimNext::FModu
 		FModuleEventTickFunction* TickFunction = Instance.FindTickFunctionByName(OtherEventName);
 		if(TickFunction == nullptr)
 		{
-			UE_LOGFMT(LogAnimation, Warning, "AddDependencyHandle: Could not find event '{EventName}' in module '{ModuleName}'", OtherEventName, Instance.GetDataInterfaceName());
+			UE_LOGFMT(LogAnimation, Warning, "AddDependencyHandle: Could not find event '{EventName}' in module '{ModuleName}'", OtherEventName, Instance.GetAssetName());
 			return;
 		}
 
@@ -440,10 +370,32 @@ void UAnimNextWorldSubsystem::RemoveModuleEventDependencyHandle(UE::AnimNext::FM
 		FModuleEventTickFunction* TickFunction = Instance.FindTickFunctionByName(InEventName);
 		if(TickFunction == nullptr)
 		{
-			UE_LOGFMT(LogAnimation, Warning, "RemoveDependencyHandle: Could not find event '{EventName}' in module '{ModuleName}'", OtherEventName, Instance.GetDataInterfaceName());
+			UE_LOGFMT(LogAnimation, Warning, "RemoveDependencyHandle: Could not find event '{EventName}' in module '{ModuleName}'", OtherEventName, Instance.GetAssetName());
 			return;
 		}
 
 		RemoveDependencyHandle(InHandle, Instance.GetObject(), *TickFunction, InEventName, InDependency);
 	}
 }
+
+EPropertyBagResult UAnimNextWorldSubsystem::SetVariableHandle(UE::AnimNext::FModuleHandle InHandle, const FAnimNextVariableReference& InVariable, const FAnimNextParamType& InType, TConstArrayView<uint8> InData)
+{
+	using namespace UE::AnimNext;
+	if (IsValidHandle(InHandle))
+	{
+		FAnimNextModuleInstance& Instance = Instances.Get(InHandle);
+		return Instance.SetProxyVariable(InVariable, InType, InData);
+	}
+	return EPropertyBagResult::PropertyNotFound;
+}
+
+EPropertyBagResult UAnimNextWorldSubsystem::WriteVariableHandle(UE::AnimNext::FModuleHandle InHandle, const FAnimNextVariableReference& InVariable, const FAnimNextParamType& InType, TFunctionRef<void(TArrayView<uint8>)> InFunction)
+{
+	using namespace UE::AnimNext;
+	if (IsValidHandle(InHandle))
+	{
+		FAnimNextModuleInstance& Instance = Instances.Get(InHandle);
+		return Instance.WriteProxyVariable(InVariable, InType, InFunction);
+	}
+	return EPropertyBagResult::PropertyNotFound;
+}
\ No newline at end of file
diff --git a/Engine/Plugins/Experimental/AnimNext/Source/AnimNext/Private/Graph/RigUnit_GetActorTransform.cpp b/Engine/Plugins/Experimental/AnimNext/Source/AnimNext/Private/Graph/RigUnit_GetActorTransform.cpp
--- a/Engine/Plugins/Experimental/AnimNext/Source/AnimNext/Private/Graph/RigUnit_GetActorTransform.cpp
+++ b/Engine/Plugins/Experimental/AnimNext/Source/AnimNext/Private/Graph/RigUnit_GetActorTransform.cpp
@@ -5,7 +5,7 @@
 #include "Component/AnimNextComponent.h"
 #include "Module/AnimNextModuleInstance.h"
 
-void FAnimNextActorTransformComponent::OnInitialize()
+FAnimNextActorTransformComponent::FAnimNextActorTransformComponent()
 {
 	FAnimNextModuleInstance& ModuleInstance = GetModuleInstance();
 	UAnimNextComponent* AnimNextComponent = Cast<UAnimNextComponent>(ModuleInstance.GetObject());
diff --git a/Engine/Plugins/Experimental/AnimNext/Source/AnimNext/Private/Graph/RigUnit_GetActorTransform.h b/Engine/Plugins/Experimental/AnimNext/Source/AnimNext/Private/Graph/RigUnit_GetActorTransform.h
--- a/Engine/Plugins/Experimental/AnimNext/Source/AnimNext/Private/Graph/RigUnit_GetActorTransform.h
+++ b/Engine/Plugins/Experimental/AnimNext/Source/AnimNext/Private/Graph/RigUnit_GetActorTransform.h
@@ -4,23 +4,24 @@
 
 #include "AnimNextExecuteContext.h"
 #include "Graph/RigUnit_AnimNextBase.h"
-#include "Module/AnimNextModuleInstanceComponent.h"
+#include "Module/UAFModuleInstanceComponent.h"
 
 #include "RigUnit_GetActorTransform.generated.h"
 
+struct FAnimNextModuleInstance;
+
 // Module component used to marshal actor transform data for RigVM units to use
 USTRUCT()
-struct FAnimNextActorTransformComponent : public FAnimNextModuleInstanceComponent
+struct FAnimNextActorTransformComponent : public FUAFModuleInstanceComponent
 {
 	GENERATED_BODY()
 
+	FAnimNextActorTransformComponent();
+	
 	// Get the transform of the actor we are bound to
 	const FTransform& GetActorTransform() const { return ActorTransform; }
 
 private:
-	// FAnimNextModuleInstanceComponent interface
-	virtual void OnInitialize() override;
-
 	// Cached actor transform
 	FTransform ActorTransform = FTransform::Identity;
 };
diff --git a/Engine/Plugins/Experimental/AnimNext/Source/AnimNext/Private/Graph/RigVMTrait_AnimNextPublicVariables.cpp b/Engine/Plugins/Experimental/AnimNext/Source/AnimNext/Private/Graph/RigVMTrait_AnimNextPublicVariables.cpp
--- a/Engine/Plugins/Experimental/AnimNext/Source/AnimNext/Private/Graph/RigVMTrait_AnimNextPublicVariables.cpp
+++ b/Engine/Plugins/Experimental/AnimNext/Source/AnimNext/Private/Graph/RigVMTrait_AnimNextPublicVariables.cpp
@@ -28,42 +28,3 @@ bool FRigVMTrait_AnimNextPublicVariables::ShouldCreatePinForProperty(const FProp
 }
 #endif
 
-namespace UE::AnimNext
-{
-
-const UAnimNextDataInterface* FPublicVariablesTraitToDataInterfaceHostAdapter::GetDataInterface() const
-{
-	return Trait.InternalAsset;
-}
-
-uint8* FPublicVariablesTraitToDataInterfaceHostAdapter::GetMemoryForVariable(int32 InVariableIndex, FName InVariableName, const FProperty* InVariableProperty) const
-{
-	// Note we dont use InVariableIndex here as we may not have bound all variables in an interface
-
-	int32 TraitVariableIndex = Trait.InternalVariableNames.Find(InVariableName);
-	if(TraitVariableIndex == INDEX_NONE)
-	{
-		// Variable not bound here
-		return nullptr;
-	}
-
-	if(!TraitScope.GetAdditionalMemoryHandles().IsValidIndex(TraitVariableIndex))
-	{
-		// Memory handle is out of bounds
-		// If this ensure fires then we have a mismatch between the variable names and the compiled memory handles, indicating a bug with the
-		// compilation of trait additional memory handles (programmatic pins)
-		ensure(false);
-		return nullptr;
-	}
-
-	const FRigVMMemoryHandle& MemoryHandle = TraitScope.GetAdditionalMemoryHandles()[TraitVariableIndex];
-	if(InVariableProperty->GetClass() != MemoryHandle.GetProperty()->GetClass())
-	{
-		UE_LOGFMT(LogAnimation, Error, "FPublicVariablesTraitToDataInterfaceHostAdapter::GetMemoryForVariable: Mismatched variable types: {Name}:{Type} vs {OtherType}", InVariableName, InVariableProperty->GetFName(), MemoryHandle.GetProperty()->GetFName());
-		return nullptr;
-	}
-
-	return const_cast<uint8*>(MemoryHandle.GetData());
-}
-
-}
\ No newline at end of file
diff --git a/Engine/Plugins/Experimental/AnimNext/Source/AnimNext/Private/Module/AnimNextActorComponentReferenceComponent.cpp b/Engine/Plugins/Experimental/AnimNext/Source/AnimNext/Private/Module/AnimNextActorComponentReferenceComponent.cpp
--- a/Engine/Plugins/Experimental/AnimNext/Source/AnimNext/Private/Module/AnimNextActorComponentReferenceComponent.cpp
+++ b/Engine/Plugins/Experimental/AnimNext/Source/AnimNext/Private/Module/AnimNextActorComponentReferenceComponent.cpp
@@ -9,8 +9,13 @@ void FAnimNextActorComponentReferenceComponent::OnInitializeHelper(UScriptStruct
 {
 	check(InScriptStruct->IsChildOf(FAnimNextActorComponentReferenceComponent::StaticStruct()));
 
-	FAnimNextModuleInstance& ModuleInstance = GetModuleInstance();
-	UAnimNextComponent* AnimNextComponent = Cast<UAnimNextComponent>(ModuleInstance.GetObject());
+	FAnimNextModuleInstance* ModuleInstance = GetModuleInstancePtr();
+	if (ModuleInstance == nullptr)
+	{
+		return;
+	}
+
+	UAnimNextComponent* AnimNextComponent = Cast<UAnimNextComponent>(ModuleInstance->GetObject());
 
 	FAnimNextModuleInstance::RunTaskOnGameThread(
 		[WeakAnimNextComponent = TWeakObjectPtr<UAnimNextComponent>(AnimNextComponent), ComponentType = ComponentType, InScriptStruct]()
@@ -30,7 +35,7 @@ void FAnimNextActorComponentReferenceComponent::OnInitializeHelper(UScriptStruct
 			UActorComponent* FoundComponent = Owner->FindComponentByClass(ComponentType);
 			PinnedAnimNextComponent->QueueTask(NAME_None, [FoundComponent, InScriptStruct](const UE::AnimNext::FModuleTaskContext& InContext)
 			{
-				InContext.TryAccessComponent(InScriptStruct, [FoundComponent](FAnimNextModuleInstanceComponent& InComponent)
+				InContext.TryAccessComponent(InScriptStruct, [FoundComponent](FUAFModuleInstanceComponent& InComponent)
 				{
 					static_cast<FAnimNextActorComponentReferenceComponent&>(InComponent).Component = FoundComponent;
 				});
diff --git a/Engine/Plugins/Experimental/AnimNext/Source/AnimNext/Private/Module/AnimNextModuleContextData.cpp b/Engine/Plugins/Experimental/AnimNext/Source/AnimNext/Private/Module/AnimNextModuleContextData.cpp
--- a/Engine/Plugins/Experimental/AnimNext/Source/AnimNext/Private/Module/AnimNextModuleContextData.cpp
+++ b/Engine/Plugins/Experimental/AnimNext/Source/AnimNext/Private/Module/AnimNextModuleContextData.cpp
@@ -5,13 +5,13 @@
 
 FAnimNextModuleContextData::FAnimNextModuleContextData(FAnimNextModuleInstance* InModuleInstance)
 	: ModuleInstance(InModuleInstance)
-	, DataInterfaceInstance(InModuleInstance)
+	, Instance(InModuleInstance)
 {
 }
 
-FAnimNextModuleContextData::FAnimNextModuleContextData(FAnimNextModuleInstance* InModuleInstance, const FAnimNextDataInterfaceInstance* InDataInterfaceInstance)
+FAnimNextModuleContextData::FAnimNextModuleContextData(FAnimNextModuleInstance* InModuleInstance, const FUAFAssetInstance* InInstance)
 	: ModuleInstance(InModuleInstance)
-	, DataInterfaceInstance(InDataInterfaceInstance)
+	, Instance(InInstance)
 {
 }
 
