AnimNext: Pass-by-ref runtime refactor

diff --git a/Engine/Plugins/Experimental/AnimNext/Source/AnimNext/Private/Module/AnimNextModuleInstance.cpp b/Engine/Plugins/Experimental/AnimNext/Source/AnimNext/Private/Module/AnimNextModuleInstance.cpp
--- a/Engine/Plugins/Experimental/AnimNext/Source/AnimNext/Private/Module/AnimNextModuleInstance.cpp
+++ b/Engine/Plugins/Experimental/AnimNext/Source/AnimNext/Private/Module/AnimNextModuleInstance.cpp
@@ -10,16 +10,15 @@
 #include "AnimNextStats.h"
 #include "Async/TaskGraphInterfaces.h"
 #include "SceneInterface.h"
+#include "UAFRigVMComponent.h"
 #include "Algo/TopologicalSort.h"
 #include "Logging/StructuredLog.h"
 #include "Misc/EnumerateRange.h"
-#include "Module/ModuleGuard.h"
 #include "Module/RigUnit_AnimNextModuleEvents.h"
-#include "Module/ProxyVariablesContext.h"
 #include "Module/ModuleTickFunction.h"
 #include "UObject/UObjectIterator.h"
-#include "Variables/IAnimNextVariableProxyHost.h"
 #include "RewindDebugger/AnimNextTrace.h"
+#include "Variables/AnimNextVariableReference.h"
 
 DEFINE_STAT(STAT_AnimNext_InitializeInstance);
 
@@ -29,11 +28,9 @@ FAnimNextModuleInstance::FAnimNextModuleInstance(
 		UAnimNextModule* InModule,
 		UObject* InObject,
 		UE::AnimNext::TPool<FAnimNextModuleInstance>* InPool,
-		IAnimNextVariableProxyHost* InProxyHost,
 		EAnimNextModuleInitMethod InInitMethod)
 	: Object(InObject)
 	, Pool(InPool)
-	, ProxyHost(InProxyHost)
 	, InitState(EInitState::NotInitialized)
 	, RunState(ERunState::NotInitialized)
 	, InitMethod(InInitMethod)
@@ -41,7 +38,7 @@ FAnimNextModuleInstance::FAnimNextModuleInstance(
 	check(InModule);
 	check(InObject);
 
-	DataInterface = InModule;
+	Asset = InModule;
 
 #if UE_ENABLE_DEBUG_DRAWING
 	if(Object && Object->GetWorld())
@@ -60,7 +57,7 @@ FAnimNextModuleInstance::~FAnimNextModuleInstance()
 #endif
 
 	Object = nullptr;
-	DataInterface = nullptr;
+	Asset = nullptr;
 	Handle.Reset();
 }
 
@@ -174,12 +171,9 @@ void FAnimNextModuleInstance::Initialize()
 
 	check(IsInGameThread());
 
-	FModuleWriteGuard Guard(this);
-
 	check(Object)
 	const UAnimNextModule* Module = GetModule();
 	check(Module);
-	check(Handle.IsValid());
 
 	UWorld* World = Object->GetWorld();
 	if(World)
@@ -286,55 +280,28 @@ void FAnimNextModuleInstance::Initialize()
 		// TODO: code in EInitState::SetupVariables phase below can probably move to FModuleEventTickFunction::Initialize
 
 		// Initialize variables
-		const int32 NumVariables = Module->VariableDefaults.GetNumPropertiesInBag();
 #if WITH_EDITOR
 		if(bIsRecreatingOnCompile)
 		{
-			Variables.MigrateToNewBagInstance(Module->VariableDefaults);
+			MigrateVariables();
 		}
 		else
 #endif
 		{
-			Variables = Module->VariableDefaults;
-		}
-
-		if(Module->GetPublicVariableDefaults().GetPropertyBagStruct())
-		{
-			PublicVariablesProxy.Data = Module->GetPublicVariableDefaults();
-			TConstArrayView<FPropertyBagPropertyDesc> ProxyDescs = PublicVariablesProxy.Data.GetPropertyBagStruct()->GetPropertyDescs();
-			PublicVariablesProxy.DirtyFlags.SetNum(ProxyDescs.Num(), false);
+			InitializeVariables();
 		}
 
-		// Initialize the RigVM context
-		ExtendedExecuteContext = Module->GetRigVMExtendedExecuteContext();
+		ProxyVariables[0].Initialize(Variables);
+		ProxyVariables[1].Initialize(Variables);
 
-		if(NumVariables > 0)
+		// Allocate compiled-in module components
 		{
-			// Setup external variables memory ptrs manually as we dont follow the pattern of owning multiple URigVMHosts like control rig.
-			// InitializeVM() is called, but only sets up handles for the defaults in the module, not for an instance
-			TArray<FRigVMExternalVariableRuntimeData> ExternalVariableRuntimeData;
-			ExternalVariableRuntimeData.Reserve(NumVariables);
-			TConstArrayView<FPropertyBagPropertyDesc> Descs = Variables.GetPropertyBagStruct()->GetPropertyDescs();
-			uint8* BasePtr = Variables.GetMutableValue().GetMemory();
-			for(int32 VariableIndex = 0; VariableIndex < NumVariables; ++VariableIndex)
+			FUAFAssetInstanceComponent::FScopedConstructorHelper ConstructorHelper(*this);
+			for(const UScriptStruct* ComponentStruct : Module->RequiredComponents)
 			{
-				ExternalVariableRuntimeData.Emplace(Descs[VariableIndex].CachedProperty->ContainerPtrToValuePtr<uint8>(BasePtr));
+				TInstancedStruct<FUAFModuleInstanceComponent> InstancedStruct(ComponentStruct);
+				ComponentMap.Add(ComponentStruct, MoveTemp(InstancedStruct));
 			}
-			ExtendedExecuteContext.ExternalVariableRuntimeData = MoveTemp(ExternalVariableRuntimeData);
-		}
-
-		// Now initialize the 'instance', cache memory handles etc. in the context
-		VM->InitializeInstance(ExtendedExecuteContext);
-
-		// Allocate compiled-in module components
-		for(UScriptStruct* ComponentStruct : Module->RequiredComponents)
-		{
-			const FName ComponentName = ComponentStruct->GetFName();
-			const int32 ComponentNameHash = GetTypeHash(ComponentName);
-
-			TInstancedStruct<FAnimNextModuleInstanceComponent> Component(ComponentStruct);
-			Component.GetMutable<FAnimNextModuleInstanceComponent>().Initialize(*this);
-			AddComponentInternal(ComponentNameHash, ComponentName, MoveTemp(Component));
 		}
 
 		TransitionToInitState(EInitState::PendingInitializeEvent);
@@ -359,14 +326,6 @@ void FAnimNextModuleInstance::Initialize()
 	}
 }
 
-void FAnimNextModuleInstance::Uninitialize()
-{
-	for(TPair<FName, TInstancedStruct<FAnimNextModuleInstanceComponent>>& ComponentPair : ComponentMap)
-	{
-		ComponentPair.Value.GetMutable<FAnimNextModuleInstanceComponent>().Uninitialize();
-	}
-}
-
 void FAnimNextModuleInstance::RemoveAllTickDependencies()
 {
 	using namespace UE::AnimNext;
@@ -377,19 +336,6 @@ void FAnimNextModuleInstance::RemoveAllTickDependencies()
 	{
 		TickFunction.RemoveAllExternalSubsequents();
 	}
-
-	if(Pool)
-	{
-		TArray<FPrerequisiteReference, TInlineAllocator<8>> PrerequisiteRefsCopy;
-		PrerequisiteRefsCopy.Append(PrerequisiteRefs);
-		for (const FPrerequisiteReference& PrerequisiteHandle : PrerequisiteRefsCopy)
-		{
-			if(FAnimNextModuleInstance* PrerequisiteInstance = Pool->TryGet(PrerequisiteHandle.Handle))
-			{
-				RemovePrerequisite(*PrerequisiteInstance);
-			}
-		}
-	}
 }
 
 void FAnimNextModuleInstance::ResetBindingsAndInstanceData()
@@ -398,8 +344,6 @@ void FAnimNextModuleInstance::ResetBindingsAndInstanceData()
 
 	check(IsInGameThread());
 
-	FModuleWriteGuard Guard(this);
-
 	TransitionToInitState(EInitState::NotInitialized);
 	TransitionToRunState(ERunState::NotInitialized);
 
@@ -413,26 +357,22 @@ void FAnimNextModuleInstance::ResetBindingsAndInstanceData()
 			check(TickFunction.ExternalSubsequents.Num() == 0);
 			TickFunction.UnRegisterTickFunction();
 		}
-		
-		Variables.Reset();
+		ProxyVariables[0].Reset();
+		ProxyVariables[1].Reset();
 	}
 
 	TickFunctions.Reset();
 
-	ExtendedExecuteContext.Reset();
+	ReleaseComponents();
 }
 
 void FAnimNextModuleInstance::QueueInputTraitEvent(FAnimNextTraitEventPtr Event)
 {
-	UE::AnimNext::FModuleWriteGuard Guard(this);
-
 	InputEventList.Push(MoveTemp(Event));
 }
 
 void FAnimNextModuleInstance::QueueOutputTraitEvent(FAnimNextTraitEventPtr Event)
 {
-	UE::AnimNext::FModuleWriteGuard Guard(this);
-
 	OutputEventList.Push(MoveTemp(Event));
 }
 
@@ -451,8 +391,6 @@ void FAnimNextModuleInstance::Enable(bool bInEnabled)
 
 	check(IsInGameThread());
 
-	FModuleWriteGuard Guard(this);
-
 	if(RunState == ERunState::Paused || RunState == ERunState::Running)
 	{
 		for (FModuleEventTickFunction& TickFunction : TickFunctions)
@@ -466,8 +404,6 @@ void FAnimNextModuleInstance::Enable(bool bInEnabled)
 
 void FAnimNextModuleInstance::TransitionToInitState(EInitState InNewState)
 {
-	UE::AnimNext::FModuleWriteGuard Guard(this);
-
 	switch(InNewState)
 	{
 	case EInitState::NotInitialized:
@@ -500,8 +436,6 @@ void FAnimNextModuleInstance::TransitionToInitState(EInitState InNewState)
 
 void FAnimNextModuleInstance::TransitionToRunState(ERunState InNewState)
 {
-	UE::AnimNext::FModuleWriteGuard Guard(this);
-
 	switch(InNewState)
 	{
 	case ERunState::Running:
@@ -522,44 +456,17 @@ void FAnimNextModuleInstance::TransitionToRunState(ERunState InNewState)
 
 void FAnimNextModuleInstance::CopyProxyVariables()
 {
-	UE::AnimNext::FModuleWriteGuard Guard(this);
-
-	// TODO: we can avoid the copies here by adopting a scheme where we:
-	// - Hold double-buffered memory handles
-	// - Update the memory handle's ptr to the currently-written double-buffered public variable on write
-	// - Swap the memory handles in ExtendedExecuteContext here
-	// 
-//	if(IAnimNextVariableProxyHost* ProxyHost = Cast<IAnimNextVariableProxyHost>(Object))
-	if(ProxyHost)
+	int32 ProxyReadIndex = 0;
 	{
-		// Flip the proxy
-		ProxyHost->FlipPublicVariablesProxy(UE::AnimNext::FProxyVariablesContext(*this));
-
-		if(PublicVariablesProxy.bIsDirty)
-		{
-			// Copy dirty properties
-			TConstArrayView<FPropertyBagPropertyDesc> ProxyDescs = Variables.GetPropertyBagStruct()->GetPropertyDescs();
-			TConstArrayView<FPropertyBagPropertyDesc> PublicProxyDescs = PublicVariablesProxy.Data.GetPropertyBagStruct()->GetPropertyDescs();
-			const uint8* SourceContainerPtr = PublicVariablesProxy.Data.GetValue().GetMemory();
-			uint8* TargetContainerPtr = Variables.GetMutableValue().GetMemory();
-			for (TConstSetBitIterator<> It(PublicVariablesProxy.DirtyFlags); It; ++It)
-			{
-				const int32 Index = It.GetIndex();
-				const FProperty* SourceProperty = PublicProxyDescs[Index].CachedProperty;
-				const FProperty* TargetProperty = ProxyDescs[Index].CachedProperty;
-				checkSlow(SourceProperty->GetClass() == TargetProperty->GetClass());
-				ProxyDescs[Index].CachedProperty->CopyCompleteValue_InContainer(TargetContainerPtr, SourceContainerPtr);
-			}
+		FWriteScopeLock WriteLock(ProxyLock);
 
-			// Reset dirty flags
-			PublicVariablesProxy.DirtyFlags.SetRange(0, PublicVariablesProxy.DirtyFlags.Num(), false);
-			PublicVariablesProxy.bIsDirty = false;
-		}
+		// Flip the write buffer index
+		ProxyReadIndex = ProxyWriteIndex;
+		ProxyWriteIndex = 1 - ProxyWriteIndex;
 	}
 
-#if ANIMNEXT_TRACE_ENABLED
-	bTracedThisFrame = false;
-#endif
+	FUAFInstanceVariableDataProxy& ProxyVariablesRead = ProxyVariables[ProxyReadIndex];
+	ProxyVariablesRead.CopyDirty();
 }
 
 #if ANIMNEXT_TRACE_ENABLED
@@ -575,7 +482,7 @@ void FAnimNextModuleInstance::Trace()
 
 const UAnimNextModule* FAnimNextModuleInstance::GetModule() const
 {
-	return CastChecked<UAnimNextModule>(DataInterface);
+	return CastChecked<UAnimNextModule>(Asset);
 }
 
 #if WITH_EDITOR
@@ -590,53 +497,14 @@ void FAnimNextModuleInstance::OnModuleCompiled()
 }
 #endif
 
-const FAnimNextModuleInstanceComponent* FAnimNextModuleInstance::TryGetComponent(int32 ComponentNameHash, FName ComponentName) const
-{
-	UE::AnimNext::FModuleWriteGuard Guard(this);
-
-	if (const TInstancedStruct<FAnimNextModuleInstanceComponent>* InstancedComponent = ComponentMap.FindByHash(ComponentNameHash, ComponentName))
-	{
-		return InstancedComponent->GetPtr();
-	}
-
-	return nullptr;
-}
-FAnimNextModuleInstanceComponent* FAnimNextModuleInstance::TryGetComponent(int32 ComponentNameHash, FName ComponentName) 
-{
-	UE::AnimNext::FModuleWriteGuard Guard(this);
-
-	if (TInstancedStruct<FAnimNextModuleInstanceComponent>* InstancedComponent = ComponentMap.FindByHash(ComponentNameHash, ComponentName))
-	{
-		return InstancedComponent->GetMutablePtr();
-	}
-
-	return nullptr;
-}
-
-FAnimNextModuleInstanceComponent& FAnimNextModuleInstance::AddComponentInternal(int32 ComponentNameHash, FName ComponentName, TInstancedStruct<FAnimNextModuleInstanceComponent>&& Component)
-{
-	UE::AnimNext::FModuleWriteGuard Guard(this);
-
-	TInstancedStruct<FAnimNextModuleInstanceComponent>& InstancedComponent = ComponentMap.AddByHash(ComponentNameHash, ComponentName, Component);
-	
-	FAnimNextModuleInstanceComponent* ModuleInstanceComponent = InstancedComponent.GetMutablePtr();
-	check(ModuleInstanceComponent != nullptr);
-
-	return *ModuleInstanceComponent;
-}
-
 #if UE_ENABLE_DEBUG_DRAWING
 FRigVMDrawInterface* FAnimNextModuleInstance::GetDebugDrawInterface()
 {
-	UE::AnimNext::FModuleWriteGuard Guard(this);
-
 	return DebugDraw ? &DebugDraw->DrawInterface : nullptr;
 }
 
 void FAnimNextModuleInstance::ShowDebugDrawing(bool bInShowDebugDrawing)
 {
-	UE::AnimNext::FModuleWriteGuard Guard(this);
-
 	if(DebugDraw)
 	{
 		DebugDraw->SetEnabled(bInShowDebugDrawing);
@@ -672,26 +540,30 @@ const UE::AnimNext::FModuleEventTickFunction* FAnimNextModuleInstance::FindTickF
 
 void FAnimNextModuleInstance::EndExecution(float InDeltaTime)
 {
-	UE::AnimNext::FModuleWriteGuard Guard(this);
-
 	// Give the module a chance to handle events
 	RaiseTraitEvents(OutputEventList);
 
 	// Give each component a chance to finalize execution
 	for(auto It = ComponentMap.CreateIterator(); It; ++It)
 	{
-		FAnimNextModuleInstanceComponent& Component = It->Value.GetMutable<FAnimNextModuleInstanceComponent>();
-		Component.OnEndExecution(InDeltaTime);
+		FUAFModuleInstanceComponent* Component = It->Value.GetMutablePtr<FUAFModuleInstanceComponent>();
+		if (Component == nullptr)
+		{
+			continue;
+		}
+		Component->OnEndExecution(InDeltaTime);
 	}
 }
 
 void FAnimNextModuleInstance::RaiseTraitEvents(const UE::AnimNext::FTraitEventList& EventList)
 {
-	UE::AnimNext::FModuleWriteGuard Guard(this);
-
 	for(auto It = ComponentMap.CreateIterator(); It; ++It)
 	{
-		FAnimNextModuleInstanceComponent& Component = It->Value.GetMutable<FAnimNextModuleInstanceComponent>();
+		FUAFModuleInstanceComponent* Component = It->Value.GetMutablePtr<FUAFModuleInstanceComponent>();
+		if (Component == nullptr)
+		{
+			continue;
+		}
 
 		// Event handlers can raise events and as such the list may change while we iterate
 		// However, if an event is added while we iterate, we will not visit it
@@ -701,110 +573,34 @@ void FAnimNextModuleInstance::RaiseTraitEvents(const UE::AnimNext::FTraitEventLi
 			const FAnimNextTraitEventPtr Event = EventList[EventIndex];
 			if (Event->IsValid())
 			{
-				Component.OnTraitEvent(*Event);
-			}
-		}
-	}
-}
-
-void FAnimNextModuleInstance::AddPrerequisite(FAnimNextModuleInstance& InPrerequisiteInstance)
-{
-	check(IsInGameThread());
-	check(&InPrerequisiteInstance != this);
-
-	UE::AnimNext::FModuleWriteGuard Guard(this);
-
-	FPrerequisiteReference* FoundReference = PrerequisiteRefs.FindByPredicate([&InPrerequisiteInstance](const FPrerequisiteReference& InPrerequisiteRef)
-	{
-		return InPrerequisiteRef.Handle == InPrerequisiteInstance.Handle;
-	});
-
-	if(FoundReference)
-	{
-		FoundReference->ReferenceCount++;
-	}
-	else
-	{
-		// Ensure all of our tick functions execute after the prerequisite's last tick function
-		TickFunctions.Last().AddPrerequisite(InPrerequisiteInstance.Object, InPrerequisiteInstance.TickFunctions.Last());
-		for(UE::AnimNext::FModuleEventTickFunction& TickFunction : TickFunctions)
-		{
-			TickFunction.AddPrerequisite(InPrerequisiteInstance.Object, InPrerequisiteInstance.TickFunctions.Last());
-		}
-
-		PrerequisiteRefs.Add({ InPrerequisiteInstance.Handle, 1});
-		InPrerequisiteInstance.SubsequentRefs.Add(GetHandle());
-	}
-}
-
-void FAnimNextModuleInstance::RemovePrerequisite(FAnimNextModuleInstance& InPrerequisiteInstance)
-{
-	check(IsInGameThread());
-	check(&InPrerequisiteInstance != this);
-
-	UE::AnimNext::FModuleWriteGuard Guard(this);
-
-	int32 FoundIndex = PrerequisiteRefs.IndexOfByPredicate([&InPrerequisiteInstance](const FPrerequisiteReference& InPrerequisiteRef)
-	{
-		return InPrerequisiteRef.Handle == InPrerequisiteInstance.Handle;
-	});
-
-	ensureAlways(FoundIndex != INDEX_NONE);	// Shouldnt really be calling this if we dont have a prerequisite already
-	if(FoundIndex != INDEX_NONE)
-	{
-		FPrerequisiteReference& FoundReference = PrerequisiteRefs[FoundIndex]; 
-		FoundReference.ReferenceCount--;
-		if(FoundReference.ReferenceCount == 0)
-		{
-			// Remove dependency on the prerequisite's last tick function
-			TickFunctions.Last().RemovePrerequisite(InPrerequisiteInstance.Object, InPrerequisiteInstance.TickFunctions.Last());
-			for(UE::AnimNext::FModuleEventTickFunction& TickFunction : TickFunctions)
-			{
-				TickFunction.RemovePrerequisite(InPrerequisiteInstance.Object, InPrerequisiteInstance.TickFunctions.Last());
+				Component->OnTraitEvent(*Event);
 			}
-
-			PrerequisiteRefs.RemoveAtSwap(FoundIndex, EAllowShrinking::No);
-
-			const int32 CountRemoved = InPrerequisiteInstance.SubsequentRefs.RemoveSwap(GetHandle(), EAllowShrinking::No);
-			check(CountRemoved == 1);
 		}
 	}
 }
 
-bool FAnimNextModuleInstance::IsPrerequisite(const FAnimNextModuleInstance& InPrerequisiteInstance) const
+EPropertyBagResult FAnimNextModuleInstance::SetProxyVariable(const FAnimNextVariableReference& InVariable, const FAnimNextParamType& InType, TConstArrayView<uint8> InData)
 {
-	UE::AnimNext::FModuleWriteGuard Guard(this);
-	return PrerequisiteRefs.ContainsByPredicate([&InPrerequisiteInstance](const FPrerequisiteReference& InPrerequisiteRef)
-	{
-		return InPrerequisiteRef.Handle == InPrerequisiteInstance.Handle;
-	});
+	FWriteScopeLock WriteLock(ProxyLock);
+	FUAFInstanceVariableDataProxy& Proxy = ProxyVariables[ProxyWriteIndex];
+	return Proxy.SetVariable(InVariable, InType, InData);
 }
 
-void FAnimNextModuleInstance::ForEachPrerequisite(TFunctionRef<void(FAnimNextModuleInstance& InPrerequisiteInstance)> InFunction) const
+EPropertyBagResult FAnimNextModuleInstance::WriteProxyVariable(const FAnimNextVariableReference& InVariable, const FAnimNextParamType& InType, TFunctionRef<void(TArrayView<uint8>)> InFunction)
 {
-	UE::AnimNext::FModuleWriteGuard Guard(this);
-
-	if(Pool == nullptr)
-	{
-		return;
-	}
-
-	for(const FPrerequisiteReference& PrerequisiteHandle : PrerequisiteRefs)
+	FWriteScopeLock WriteLock(ProxyLock);
+	FUAFInstanceVariableDataProxy& Proxy = ProxyVariables[ProxyWriteIndex];
+	TArrayView<uint8> Data;
+	EPropertyBagResult Result = Proxy.WriteVariable(InVariable, InType, Data);
+	if (Result == EPropertyBagResult::Success)
 	{
-		FAnimNextModuleInstance* PrerequisiteInstance = Pool->TryGet(PrerequisiteHandle.Handle);
-		if(PrerequisiteInstance == nullptr)
-		{
-			continue;
-		}
-
-		InFunction(*PrerequisiteInstance);
+		InFunction(Data);
 	}
+	return Result;
 }
 
 void FAnimNextModuleInstance::RunRigVMEvent(FName InEventName, float InDeltaTime)
 {
-	UE::AnimNext::FModuleWriteGuard PrerequisiteGuard(this);
-
 	URigVM* VM = GetModule()->RigVM;
 	if(VM == nullptr)
 	{
@@ -816,6 +612,8 @@ void FAnimNextModuleInstance::RunRigVMEvent(FName InEventName, float InDeltaTime
 		return;
 	}
 
+	FUAFRigVMComponent& RigVMComponent = GetComponent<FUAFRigVMComponent>();
+	FRigVMExtendedExecuteContext& ExtendedExecuteContext = RigVMComponent.GetExtendedExecuteContext();
 	check(ExtendedExecuteContext.VMHash == VM->GetVMHash());
 
 	FAnimNextExecuteContext& AnimNextContext = ExtendedExecuteContext.GetPublicDataSafe<FAnimNextExecuteContext>();
@@ -839,15 +637,11 @@ void FAnimNextModuleInstance::RunRigVMEvent(FName InEventName, float InDeltaTime
 
 TArrayView<UE::AnimNext::FModuleEventTickFunction> FAnimNextModuleInstance::GetTickFunctions()
 {
-	UE::AnimNext::FModuleWriteGuard PrerequisiteGuard(this);
-
 	return TickFunctions;
 }
 
 UE::AnimNext::FModuleEventTickFunction* FAnimNextModuleInstance::FindFirstUserTickFunction()
 {
-	UE::AnimNext::FModuleWriteGuard PrerequisiteGuard(this);
-
 	UE::AnimNext::FModuleEventTickFunction* FoundTickFunction = TickFunctions.FindByPredicate([](const UE::AnimNext::FModuleEventTickFunction& InTickFunction)
 	{
 		return InTickFunction.bFirstUserEvent;
@@ -902,7 +696,7 @@ void FAnimNextModuleInstance::QueueTask(FName InEventName, TUniqueFunction<void(
 	}
 	else
 	{
-		UE_LOGFMT(LogAnimation, Warning, "QueueTask: Could not find event '{EventName}' in module '{ModuleName}'", InEventName, GetDataInterfaceName());
+		UE_LOGFMT(LogAnimation, Warning, "QueueTask: Could not find event '{EventName}' in module '{ModuleName}'", InEventName, GetAssetName());
 	}
 }
 
diff --git a/Engine/Plugins/Experimental/AnimNext/Source/AnimNext/Private/Module/AnimNextSkeletalMeshComponentReferenceComponent.cpp b/Engine/Plugins/Experimental/AnimNext/Source/AnimNext/Private/Module/AnimNextSkeletalMeshComponentReferenceComponent.cpp
--- a/Engine/Plugins/Experimental/AnimNext/Source/AnimNext/Private/Module/AnimNextSkeletalMeshComponentReferenceComponent.cpp
+++ b/Engine/Plugins/Experimental/AnimNext/Source/AnimNext/Private/Module/AnimNextSkeletalMeshComponentReferenceComponent.cpp
@@ -7,14 +7,10 @@
 FAnimNextSkeletalMeshComponentReferenceComponent::FAnimNextSkeletalMeshComponentReferenceComponent()
 {
 	ComponentType = USkeletalMeshComponent::StaticClass();
+	OnInitializeHelper(StaticStruct());
 }
 
 USkeletalMeshComponent* FAnimNextSkeletalMeshComponentReferenceComponent::GetComponent() const
 {
 	return Cast<USkeletalMeshComponent>(Component);
 }
-
-void FAnimNextSkeletalMeshComponentReferenceComponent::OnInitialize()
-{
-	OnInitializeHelper(StaticStruct());
-}
\ No newline at end of file
diff --git a/Engine/Plugins/Experimental/AnimNext/Source/AnimNext/Private/Module/ModuleTaskContext.cpp b/Engine/Plugins/Experimental/AnimNext/Source/AnimNext/Private/Module/ModuleTaskContext.cpp
--- a/Engine/Plugins/Experimental/AnimNext/Source/AnimNext/Private/Module/ModuleTaskContext.cpp
+++ b/Engine/Plugins/Experimental/AnimNext/Source/AnimNext/Private/Module/ModuleTaskContext.cpp
@@ -3,30 +3,24 @@
 #include "Module/ModuleTaskContext.h"
 
 #include "Module/AnimNextModuleInstance.h"
-#include "Module/ModuleGuard.h"
 
 namespace UE::AnimNext
 {
 
 FModuleTaskContext::FModuleTaskContext(FAnimNextModuleInstance& InModuleInstance)
-	: ModuleInstance(&InModuleInstance)
+	: FInstanceTaskContext(InModuleInstance)
+	, ModuleInstance(InModuleInstance)
 {
 }
 
 void FModuleTaskContext::QueueInputTraitEvent(FAnimNextTraitEventPtr Event) const
 {
-	FModuleWriteGuard Guard(ModuleInstance);
-
-	ModuleInstance->QueueInputTraitEvent(MoveTemp(Event));
+	ModuleInstance.QueueInputTraitEvent(MoveTemp(Event));
 }
 
-void FModuleTaskContext::TryAccessComponent(UScriptStruct* InComponentType, TFunctionRef<void(FAnimNextModuleInstanceComponent&)> InFunction) const
+void FModuleTaskContext::TryAccessComponent(UScriptStruct* InComponentType, TFunctionRef<void(FUAFModuleInstanceComponent&)> InFunction) const
 {
-	FModuleWriteGuard Guard(ModuleInstance);
-
-	FName ComponentName = InComponentType->GetFName();
-	uint32 ComponentNameHash = GetTypeHash(ComponentName);
-	FAnimNextModuleInstanceComponent* Component = ModuleInstance->TryGetComponent(ComponentNameHash, ComponentName);
+	FUAFModuleInstanceComponent* Component = static_cast<FUAFModuleInstanceComponent*>(ModuleInstance.TryGetComponent(InComponentType));
 	if(Component == nullptr)
 	{
 		return;
@@ -37,7 +31,7 @@ void FModuleTaskContext::TryAccessComponent(UScriptStruct* InComponentType, TFun
 
 FAnimNextModuleInstance* const FModuleTaskContext::GetModuleInstance() const
 {
-	return ModuleInstance;
+	return &ModuleInstance;
 }
 
 }
diff --git a/Engine/Plugins/Experimental/AnimNext/Source/AnimNext/Private/Module/ModuleTickFunction.cpp b/Engine/Plugins/Experimental/AnimNext/Source/AnimNext/Private/Module/ModuleTickFunction.cpp
--- a/Engine/Plugins/Experimental/AnimNext/Source/AnimNext/Private/Module/ModuleTickFunction.cpp
+++ b/Engine/Plugins/Experimental/AnimNext/Source/AnimNext/Private/Module/ModuleTickFunction.cpp
@@ -9,12 +9,9 @@
 #include "Component/AnimNextWorldSubsystem.h"
 #include "Module/AnimNextModule.h"
 #include "Module/ModuleTaskContext.h"
-#include "Module/ProxyVariablesContext.h"
 #include "Module/RigUnit_AnimNextModuleEvents.h"
 #include "Module/ModuleEvents.h"
 #include "TraitCore/TraitEventList.h"
-#include "Module/ModuleGuard.h"
-
 
 namespace UE::AnimNext
 {
@@ -48,8 +45,6 @@ void FModuleEventTickFunction::EndTick(float DeltaTime)
 {
 	SCOPED_NAMED_EVENT(AnimNext_Module_EndTick, FColor::Orange);
 
-	FModuleWriteGuard Guard(ModuleInstance);
-
 	// Give module a chance to finish up processing
 	ModuleInstance->EndExecution(DeltaTime);
 
@@ -192,8 +187,6 @@ void FModuleEventTickFunction::Run(float InDeltaTime)
 {
 	SCOPED_NAMED_EVENT(AnimNext_Module_EventTick, FColor::Orange);
 
-	FModuleWriteGuard Guard(ModuleInstance);
-
 	while (!PreExecuteTasks.IsEmpty())
 	{
 		TOptional<TUniqueFunction<void(const FModuleTaskContext&)>> Function = PreExecuteTasks.Dequeue();
@@ -203,6 +196,10 @@ void FModuleEventTickFunction::Run(float InDeltaTime)
 
 	if (bFirstUserEvent)
 	{
+#if ANIMNEXT_TRACE_ENABLED
+		ModuleInstance->bTracedThisFrame = false;
+#endif
+
 		if (bRunBindingsEvent)
 		{
 			// Execute any WT bindings
diff --git a/Engine/Plugins/Experimental/AnimNext/Source/AnimNext/Private/Param/ParamType.cpp b/Engine/Plugins/Experimental/AnimNext/Source/AnimNext/Private/Param/ParamType.cpp
--- a/Engine/Plugins/Experimental/AnimNext/Source/AnimNext/Private/Param/ParamType.cpp
+++ b/Engine/Plugins/Experimental/AnimNext/Source/AnimNext/Private/Param/ParamType.cpp
@@ -8,13 +8,6 @@
 #include "RigVMCore/RigVMTemplate.h"
 #include "UObject/TextProperty.h"
 
-FAnimNextParamType::FAnimNextParamType(EValueType InValueType, EContainerType InContainerType, const UObject* InValueTypeObject)
-	: ValueTypeObject(InValueTypeObject)
-	, ValueType(InValueType)
-	, ContainerType(InContainerType)
-{
-}
-
 FRigVMTemplateArgumentType FAnimNextParamType::ToRigVMTemplateArgument() const
 {
 	auto GetTypeInternal = [this]()
@@ -243,7 +236,7 @@ FAnimNextParamType FAnimNextParamType::FromProperty(const FProperty* InProperty)
 			if(Class == UClass::StaticClass())
 			{
 				Type.ValueType = EPropertyBagPropertyType::Class;
-				Type.ValueTypeObject = Class;
+				Type.ValueTypeObject = UObject::StaticClass();
 			}
 			else
 			{
diff --git a/Engine/Plugins/Experimental/AnimNext/Source/AnimNext/Private/Param/ParamUtils.cpp b/Engine/Plugins/Experimental/AnimNext/Source/AnimNext/Private/Param/ParamUtils.cpp
--- a/Engine/Plugins/Experimental/AnimNext/Source/AnimNext/Private/Param/ParamUtils.cpp
+++ b/Engine/Plugins/Experimental/AnimNext/Source/AnimNext/Private/Param/ParamUtils.cpp
@@ -26,6 +26,8 @@ FParamCompatibility FParamUtils::GetCompatibility(const FAnimNextParamType& InLH
 	case EPropertyBagPropertyType::Byte:
 		switch (InLHS.GetValueType())
 		{
+		case EPropertyBagPropertyType::Bool:
+			return EParamCompatibility::Incompatible_DataLoss;
 		case EPropertyBagPropertyType::Byte:
 			return EParamCompatibility::Compatible_Equal;
 		case EPropertyBagPropertyType::Int32:
@@ -41,6 +43,8 @@ FParamCompatibility FParamUtils::GetCompatibility(const FAnimNextParamType& InLH
 	case EPropertyBagPropertyType::Int32:
 		switch (InLHS.GetValueType())
 		{
+		case EPropertyBagPropertyType::Bool:
+			return EParamCompatibility::Incompatible_DataLoss;
 		case EPropertyBagPropertyType::Byte:
 			return EParamCompatibility::Incompatible_DataLoss;
 		case EPropertyBagPropertyType::Int32:
@@ -56,6 +60,8 @@ FParamCompatibility FParamUtils::GetCompatibility(const FAnimNextParamType& InLH
 	case EPropertyBagPropertyType::Int64:
 		switch (InLHS.GetValueType())
 		{
+		case EPropertyBagPropertyType::Bool:
+			return EParamCompatibility::Incompatible_DataLoss;
 		case EPropertyBagPropertyType::Byte:
 			return EParamCompatibility::Incompatible_DataLoss;
 		case EPropertyBagPropertyType::Int32:
@@ -71,6 +77,8 @@ FParamCompatibility FParamUtils::GetCompatibility(const FAnimNextParamType& InLH
 	case EPropertyBagPropertyType::Float:
 		switch (InLHS.GetValueType())
 		{
+		case EPropertyBagPropertyType::Bool:
+			return EParamCompatibility::Incompatible_DataLoss;
 		case EPropertyBagPropertyType::Byte:
 			return EParamCompatibility::Incompatible_DataLoss;
 		case EPropertyBagPropertyType::Int32:
@@ -86,6 +94,8 @@ FParamCompatibility FParamUtils::GetCompatibility(const FAnimNextParamType& InLH
 	case EPropertyBagPropertyType::Double:
 		switch (InLHS.GetValueType())
 		{
+		case EPropertyBagPropertyType::Bool:
+			return EParamCompatibility::Incompatible_DataLoss;
 		case EPropertyBagPropertyType::Byte:
 			return EParamCompatibility::Incompatible_DataLoss;
 		case EPropertyBagPropertyType::Int32:
diff --git a/Engine/Plugins/Experimental/AnimNext/Source/AnimNext/Private/RewindDebugger/AnimNextTrace.cpp b/Engine/Plugins/Experimental/AnimNext/Source/AnimNext/Private/RewindDebugger/AnimNextTrace.cpp
--- a/Engine/Plugins/Experimental/AnimNext/Source/AnimNext/Private/RewindDebugger/AnimNextTrace.cpp
+++ b/Engine/Plugins/Experimental/AnimNext/Source/AnimNext/Private/RewindDebugger/AnimNextTrace.cpp
@@ -4,7 +4,7 @@
 
 #include "ObjectTrace.h"
 #include "Chaos/ParticleDirtyFlags.h"
-#include "DataInterface/AnimNextDataInterfaceInstance.h"
+#include "UAFAssetInstance.h"
 #include "Misc/HashBuilder.h"
 #include "Module/AnimNextModuleInstance.h"
 #include "Serialization/MemoryWriter.h"
@@ -68,10 +68,10 @@ void FAnimNextTrace::Reset()
 	GTracedPropertyDescs.Empty();
 }
 
-void FAnimNextTrace::OutputAnimNextInstance(const FAnimNextDataInterfaceInstance* DataInterface, const UObject* OuterObject)
+void FAnimNextTrace::OutputAnimNextInstance(const FUAFAssetInstance* AnimNextInstance, const UObject* OuterObject)
 {
 	bool bChannelEnabled = UE_TRACE_CHANNELEXPR_IS_ENABLED(AnimNextChannel);
-	if (!bChannelEnabled || DataInterface == nullptr)
+	if (!bChannelEnabled || AnimNextInstance == nullptr)
 	{
 		return;
 	}
@@ -81,7 +81,7 @@ void FAnimNextTrace::OutputAnimNextInstance(const FAnimNextDataInterfaceInstance
 		return;
 	}
 	
-	uint64 InstanceId = DataInterface->GetUniqueId();
+	uint64 InstanceId = AnimNextInstance->GetUniqueId();
 
 	bool bTrace = false;
 
@@ -100,21 +100,21 @@ void FAnimNextTrace::OutputAnimNextInstance(const FAnimNextDataInterfaceInstance
 
 	if (bTrace)
 	{
-		if (DataInterface->GetHost())
+		if (AnimNextInstance->GetHost())
 		{
-			OutputAnimNextInstance(DataInterface->GetHost(), OuterObject);
+			OutputAnimNextInstance(AnimNextInstance->GetHost(), OuterObject);
 		}
-		
+
 		TRACE_OBJECT(OuterObject);
-		TRACE_OBJECT(DataInterface->GetDataInterface());
+		TRACE_OBJECT(AnimNextInstance->GetAsset<UAnimNextRigVMAsset>());
 		
-		const FAnimNextDataInterfaceInstance* HostInstance = DataInterface->GetHost();
+		const FUAFAssetInstance* HostInstance = AnimNextInstance->GetHost();
 
 		UE_TRACE_LOG(AnimNext, Instance, AnimNextChannel)
 			<< Instance.InstanceId(InstanceId)
 			<< Instance.OuterObjectId(FObjectTrace::GetObjectId(OuterObject))
 			<< Instance.HostInstanceId(HostInstance ? HostInstance->GetUniqueId() : 0)
-			<< Instance.AssetId(FObjectTrace::GetObjectId(DataInterface->GetDataInterface()));
+			<< Instance.AssetId(FObjectTrace::GetObjectId(AnimNextInstance->GetAsset<UAnimNextRigVMAsset>()));
 	}
 }
 
@@ -134,7 +134,7 @@ uint32 GetPropertyDescHash(const TConstArrayView<FPropertyBagPropertyDesc>& Desc
 	return HashBuilder.GetHash();
 }
 
-void FAnimNextTrace::OutputAnimNextVariables(const FAnimNextDataInterfaceInstance* Instance, const UObject* OuterObject)
+void FAnimNextTrace::OutputAnimNextVariables(const FUAFAssetInstance* Instance, const UObject* OuterObject)
 {
 	bool bChannelEnabled = UE_TRACE_CHANNELEXPR_IS_ENABLED(AnimNextChannel);
 	if (!bChannelEnabled || Instance == nullptr || OuterObject == nullptr)
@@ -150,76 +150,95 @@ void FAnimNextTrace::OutputAnimNextVariables(const FAnimNextDataInterfaceInstanc
 	}
 
 	OutputAnimNextInstance(Instance, OuterObject);
-	
-	FInstancedPropertyBag Variables = Instance->GetVariables();
-	if (Variables.GetPropertyBagStruct())
+
+	if (Instance->Variables.OwnedVariableContainers.Num() > 0)
 	{
-		TConstArrayView<FPropertyBagPropertyDesc> PropertyDescs = Variables.GetPropertyBagStruct()->GetPropertyDescs();
-		for(int i =0; i<PropertyDescs.Num(); i++)
+		for (const TSharedRef<FUAFInstanceVariableContainer>& OwnedVariableSet : Instance->Variables.OwnedVariableContainers)
 		{
-			uint8* ValueMemory = Instance->GetMemoryForVariable(i, PropertyDescs[i].Name, PropertyDescs[i].CachedProperty);
-			PropertyDescs[i].CachedProperty->CopySingleValue(const_cast<uint8*>(Variables.GetValue().GetMemory()) +  PropertyDescs[i].CachedProperty->GetOffset_ForInternal(), ValueMemory);
+			OutputAnimNextVariableSet(OwnedVariableSet, InstanceId, OuterObject);
 		}
-		
-		uint32 PropertyDescHash = GetPropertyDescHash(PropertyDescs);
+	}
+	else
+	{
+		UE_TRACE_LOG(AnimNext, InstanceVariables, AnimNextChannel)
+				<< InstanceVariables.Cycle(FPlatformTime::Cycles64())
+				<< InstanceVariables.RecordingTime(FObjectTrace::GetWorldElapsedTime(OuterObject->GetWorld()))
+				<< InstanceVariables.InstanceId(InstanceId);
+	}
+}
+
+void FAnimNextTrace::OutputAnimNextVariableSet(const TSharedRef<FUAFInstanceVariableContainer>& InVariableSet, uint64 InstanceId, const UObject* OuterObject)
+{
+	switch (InVariableSet->VariablesContainer.GetIndex())
+	{
+	case FUAFInstanceVariableContainer::FVariablesContainerType::IndexOfType<FInstancedPropertyBag>():
 		{
-			FRWScopeLock RWScopeLock(GTracedPropertiesRWLock, SLT_ReadOnly);
-			TArray<uint8> ArchiveData;
-			ArchiveData.Reserve(1024 * 10);
-			
-			if (!GTracedPropertyDescs.Contains(PropertyDescHash))
-			{
-				// Upgrade to write lock
-				RWScopeLock.ReleaseReadOnlyLockAndAcquireWriteLock_USE_WITH_CAUTION();
+			// TODO: Accomodate overrides
 
-				// Check in case someone just serialized before we aquired the write lock
+			FInstancedPropertyBag& InstancedPropertyBag = InVariableSet->VariablesContainer.Get<FInstancedPropertyBag>();
+			TConstArrayView<FPropertyBagPropertyDesc> PropertyDescs = InstancedPropertyBag.GetPropertyBagStruct()->GetPropertyDescs();
+			uint32 PropertyDescHash = GetPropertyDescHash(PropertyDescs);
+			{
+				FRWScopeLock RWScopeLock(GTracedPropertiesRWLock, SLT_ReadOnly);
+				TArray<uint8> ArchiveData;
+				ArchiveData.Reserve(1024 * 10);
+				
 				if (!GTracedPropertyDescs.Contains(PropertyDescHash))
 				{
-					GTracedPropertyDescs.Add(PropertyDescHash, Variables.GetPropertyBagStruct());
-					
-					// Serialize prop descs
-					TArray<FPropertyBagPropertyDesc, TInlineAllocator<64>> PropertyDescriptions;
-					PropertyDescriptions.Reserve(PropertyDescs.Num());
-					PropertyDescriptions = PropertyDescs;	// need to copy since we can't serialize an array view
-
-					FMemoryWriter WriterArchive(ArchiveData);
-					FObjectAndNameAsStringProxyArchive Archive(WriterArchive, /*bInLoadIfFindFails*/true);
-					Archive.UsingCustomVersion(CustomVersionGUID);
-					Archive.UsingCustomVersion(FPropertyBagCustomVersion::GUID);
-					Archive << PropertyDescriptions;
-
-					UE_TRACE_LOG(AnimNext, InstanceVariableDescriptions, AnimNextChannel)
-						<< InstanceVariableDescriptions.VariableDescriptionHash(PropertyDescHash)
-						<< InstanceVariableDescriptions.VariableDescriptionData(ArchiveData.GetData(), ArchiveData.Num());
+					// Upgrade to write lock
+					RWScopeLock.ReleaseReadOnlyLockAndAcquireWriteLock_USE_WITH_CAUTION();
+
+					// Check in case someone just serialized before we aquired the write lock
+					if (!GTracedPropertyDescs.Contains(PropertyDescHash))
+					{
+						GTracedPropertyDescs.Add(PropertyDescHash, InstancedPropertyBag.GetPropertyBagStruct());
+						
+						// Serialize prop descs
+						TArray<FPropertyBagPropertyDesc, TInlineAllocator<64>> PropertyDescriptions;
+						PropertyDescriptions.Reserve(PropertyDescs.Num());
+						PropertyDescriptions = PropertyDescs;	// need to copy since we can't serialize an array view
+
+						FMemoryWriter WriterArchive(ArchiveData);
+						FObjectAndNameAsStringProxyArchive Archive(WriterArchive, /*bInLoadIfFindFails*/true);
+						Archive.UsingCustomVersion(CustomVersionGUID);
+						Archive.UsingCustomVersion(FPropertyBagCustomVersion::GUID);
+						Archive << PropertyDescriptions;
+
+						UE_TRACE_LOG(AnimNext, InstanceVariableDescriptions, AnimNextChannel)
+							<< InstanceVariableDescriptions.VariableDescriptionHash(PropertyDescHash)
+							<< InstanceVariableDescriptions.VariableDescriptionData(ArchiveData.GetData(), ArchiveData.Num());
+					}
 				}
-			}
 
-			ArchiveData.SetNum(0, EAllowShrinking::No);
-			FMemoryWriter WriterArchive(ArchiveData);
-			FObjectAndNameAsStringProxyArchive Archive(WriterArchive, /*bInLoadIfFindFails*/true);
-			UPropertyBag* PropertyBag = const_cast<UPropertyBag*>(Variables.GetPropertyBagStruct());
-			if (PropertyBag != nullptr)
-			{
-				PropertyBag->SerializeItem(Archive, Variables.GetMutableValue().GetMemory(), nullptr);
-			}
+				ArchiveData.SetNum(0, EAllowShrinking::No);
+				FMemoryWriter WriterArchive(ArchiveData);
+				FObjectAndNameAsStringProxyArchive Archive(WriterArchive, /*bInLoadIfFindFails*/true);
+				UPropertyBag* PropertyBag = const_cast<UPropertyBag*>(InstancedPropertyBag.GetPropertyBagStruct());
+				if (PropertyBag != nullptr)
+				{
+					PropertyBag->SerializeItem(Archive, InstancedPropertyBag.GetMutableValue().GetMemory(), nullptr);
+				}
 
-			UE_TRACE_LOG(AnimNext, InstanceVariables, AnimNextChannel)
-					<< InstanceVariables.Cycle(FPlatformTime::Cycles64())
-					<< InstanceVariables.RecordingTime(FObjectTrace::GetWorldElapsedTime(OuterObject->GetWorld()))
-					<< InstanceVariables.InstanceId(InstanceId)
-					<< InstanceVariables.VariableDescriptionHash(PropertyDescHash)
-					<< InstanceVariables.VariableData(ArchiveData.GetData(), ArchiveData.Num());
+				UE_TRACE_LOG(AnimNext, InstanceVariables, AnimNextChannel)
+						<< InstanceVariables.Cycle(FPlatformTime::Cycles64())
+						<< InstanceVariables.RecordingTime(FObjectTrace::GetWorldElapsedTime(OuterObject->GetWorld()))
+						<< InstanceVariables.InstanceId(InstanceId)
+						<< InstanceVariables.VariableDescriptionHash(PropertyDescHash)
+						<< InstanceVariables.VariableData(ArchiveData.GetData(), ArchiveData.Num());
+			}
+			break;
 		}
-	}
-	else
-	{
-		UE_TRACE_LOG(AnimNext, InstanceVariables, AnimNextChannel)
-				<< InstanceVariables.Cycle(FPlatformTime::Cycles64())
-				<< InstanceVariables.RecordingTime(FObjectTrace::GetWorldElapsedTime(OuterObject->GetWorld()))
-				<< InstanceVariables.InstanceId(InstanceId);
+		case FUAFInstanceVariableContainer::FVariablesContainerType::IndexOfType<FInstancedStruct>():
+		{
+			// TODO: Trace native instanced structs
+			break;
+		}
+	default:
+		checkNoEntry();
+		break;
 	}
 }
-	
+
 }
 
 #endif
diff --git a/Engine/Plugins/Experimental/AnimNext/Source/AnimNext/Private/UAFAssetInstance.cpp b/Engine/Plugins/Experimental/AnimNext/Source/AnimNext/Private/UAFAssetInstance.cpp
new file mode 100644
--- /dev/null
+++ b/Engine/Plugins/Experimental/AnimNext/Source/AnimNext/Private/UAFAssetInstance.cpp
@@ -0,0 +1,98 @@
+// Copyright Epic Games, Inc. All Rights Reserved.
+
+#include "UAFAssetInstance.h"
+
+#include "Logging/StructuredLog.h"
+#include "Variables/UAFInstanceVariableContainer.h"
+#include "Variables/VariableOverrides.h"
+#include "Variables/VariableOverridesCollection.h"
+#include "UAFAssetInstanceComponent.h"
+
+#if ANIMNEXT_TRACE_ENABLED
+volatile int64 FUAFAssetInstance::NextUniqueId = 1;
+#endif
+
+void FUAFAssetInstance::InitializeVariables(const TSharedPtr<const UE::AnimNext::FVariableOverridesCollection>& InOverrides)
+{
+	Variables.Initialize(*this, InOverrides);
+}
+
+#if WITH_EDITOR
+void FUAFAssetInstance::MigrateVariables()
+{
+	Variables.Migrate();
+}
+#endif
+
+
+void FUAFAssetInstance::ReleaseComponents()
+{
+	ComponentMap.Empty();
+	CachedRigVMComponent = nullptr;
+}
+
+FUAFAssetInstance::FUAFAssetInstance()
+{
+#if ANIMNEXT_TRACE_ENABLED
+	UniqueId = FPlatformAtomics::InterlockedIncrement(&NextUniqueId);
+#endif
+}
+
+#if DO_CHECK
+bool FUAFAssetInstance::LayoutMatches(const FInstancedPropertyBag& InPropertyBag) const
+{
+	const UPropertyBag* PropertyBagStruct = InPropertyBag.GetPropertyBagStruct();
+	if (PropertyBagStruct == nullptr)
+	{
+		return Variables.NumInternalVariables == 0;
+	}
+
+	if (Variables.NumInternalVariables != InPropertyBag.GetNumPropertiesInBag())
+	{
+		return false;
+	}
+
+	bool bMatches = true;
+	TConstArrayView<FPropertyBagPropertyDesc> Descs = PropertyBagStruct->GetPropertyDescs();
+	Variables.ForEachVariable([&Descs, &bMatches](FName InName, const FAnimNextParamType& InType, int32 InVariableIndex)
+	{
+		if (!Descs.IsValidIndex(InVariableIndex))
+		{
+			bMatches = false;
+			return false;
+		}
+
+		const FPropertyBagPropertyDesc& Desc = Descs[InVariableIndex];
+		if (Desc.Name != InName)
+		{
+			bMatches = false;
+			return false;
+		}
+
+		if (Desc.ValueType != InType.GetValueType() ||
+			Desc.ContainerTypes.GetFirstContainerType() != InType.GetContainerType() ||
+			Desc.ValueTypeObject != InType.GetValueTypeObject())
+		{
+			bMatches = false;
+			return false;
+		}
+
+		return true;
+	});
+
+	return bMatches;
+}
+#endif
+
+FUAFAssetInstanceComponent* FUAFAssetInstance::TryGetComponent(const UScriptStruct* InStruct)
+{
+	checkf(InStruct->IsChildOf<FUAFAssetInstanceComponent>(), TEXT("ComponentType type must derive from FUAFAssetInstanceComponent"));
+
+	const int32 ComponentStructHash = GetTypeHash(InStruct);
+
+	if (TInstancedStruct<FUAFAssetInstanceComponent>* InstanceComponent = ComponentMap.FindByHash(ComponentStructHash, InStruct))
+	{
+		return InstanceComponent->GetMutablePtr<FUAFAssetInstanceComponent>();
+	}
+	return nullptr;
+}
\ No newline at end of file
diff --git a/Engine/Plugins/Experimental/AnimNext/Source/AnimNext/Private/UAFAssetInstanceComponent.cpp b/Engine/Plugins/Experimental/AnimNext/Source/AnimNext/Private/UAFAssetInstanceComponent.cpp
new file mode 100644
--- /dev/null
+++ b/Engine/Plugins/Experimental/AnimNext/Source/AnimNext/Private/UAFAssetInstanceComponent.cpp
@@ -0,0 +1,23 @@
+﻿// Copyright Epic Games, Inc. All Rights Reserved.
+
+#include "UAFAssetInstanceComponent.h"
+
+namespace UE::AnimNext::Private
+{
+	static thread_local FUAFAssetInstance* GCurrentAssetInstance = nullptr;
+}
+
+FUAFAssetInstanceComponent::FScopedConstructorHelper::FScopedConstructorHelper(FUAFAssetInstance& InInstance)
+{
+	UE::AnimNext::Private::GCurrentAssetInstance = &InInstance;
+}
+
+FUAFAssetInstanceComponent::FScopedConstructorHelper::~FScopedConstructorHelper()
+{
+	UE::AnimNext::Private::GCurrentAssetInstance = nullptr;
+}
+
+FUAFAssetInstanceComponent::FUAFAssetInstanceComponent()
+{
+	Instance = UE::AnimNext::Private::GCurrentAssetInstance;
+}
\ No newline at end of file
diff --git a/Engine/Plugins/Experimental/AnimNext/Source/AnimNext/Private/UAFRigVMComponent.cpp b/Engine/Plugins/Experimental/AnimNext/Source/AnimNext/Private/UAFRigVMComponent.cpp
new file mode 100644
--- /dev/null
+++ b/Engine/Plugins/Experimental/AnimNext/Source/AnimNext/Private/UAFRigVMComponent.cpp
@@ -0,0 +1,49 @@
+﻿// Copyright Epic Games, Inc. All Rights Reserved.
+
+#include "UAFRigVMComponent.h"
+
+#include "AnimNextRigVMAsset.h"
+#include "UAFAssetInstance.h"
+#include "Variables/UAFInstanceVariableContainer.h"
+
+FUAFRigVMComponent::FUAFRigVMComponent()
+{
+	FUAFAssetInstance* AssetInstance = GetAssetInstancePtr();
+	if (AssetInstance == nullptr)
+	{
+		return;
+	}
+
+	const UAnimNextRigVMAsset* Asset = AssetInstance->GetAsset<UAnimNextRigVMAsset>();
+	URigVM* VM = Asset->RigVM;
+	if (VM == nullptr)
+	{
+		return;
+	}
+
+	// Initialize the RigVM context
+	ExtendedExecuteContext = Asset->GetRigVMExtendedExecuteContext();
+
+	// Hookup the runtime data ptrs
+	const TArray<FRigVMExternalVariableDef>& ExternalVariableDefs = VM->GetExternalVariableDefs();
+	const int32 NumExternalVariables = ExternalVariableDefs.Num();
+	check(AssetInstance->Variables.NumInternalVariables == NumExternalVariables);
+	if(NumExternalVariables > 0)
+	{
+		TArray<FRigVMExternalVariableRuntimeData> ExternalVariableRuntimeData;
+		ExternalVariableRuntimeData.Reserve(AssetInstance->Variables.NumInternalVariables);
+		const int32 NumInternalSharedVariableContainers = AssetInstance->Variables.InternalSharedVariableContainers.Num();
+		int32 CurrentBaseIndex = 0;
+		for (int32 SharedVariableSetIndex = 0; SharedVariableSetIndex < NumInternalSharedVariableContainers; ++SharedVariableSetIndex)
+		{
+			const TSharedRef<FUAFInstanceVariableContainer>& SharedVariableSet = AssetInstance->Variables.InternalSharedVariableContainers[SharedVariableSetIndex].Pin().ToSharedRef();
+			TConstArrayView<FRigVMExternalVariableDef> DefsForSet(ExternalVariableDefs.GetData() + CurrentBaseIndex, SharedVariableSet->NumVariables);
+			SharedVariableSet->GetRigVMMemoryForVariables(DefsForSet, ExternalVariableRuntimeData);
+			CurrentBaseIndex += SharedVariableSet->NumVariables;
+		}
+		ExtendedExecuteContext.ExternalVariableRuntimeData = MoveTemp(ExternalVariableRuntimeData);
+	}
+
+	// Now initialize the 'instance', cache memory handles etc. in the context
+	VM->InitializeInstance(ExtendedExecuteContext);
+}
\ No newline at end of file
diff --git a/Engine/Plugins/Experimental/AnimNext/Source/AnimNext/Private/Variables/AnimNextSharedVariables.cpp b/Engine/Plugins/Experimental/AnimNext/Source/AnimNext/Private/Variables/AnimNextSharedVariables.cpp
new file mode 100644
--- /dev/null
+++ b/Engine/Plugins/Experimental/AnimNext/Source/AnimNext/Private/Variables/AnimNextSharedVariables.cpp
@@ -0,0 +1,15 @@
+﻿// Copyright Epic Games, Inc. All Rights Reserved.
+
+#include "Variables/AnimNextSharedVariables.h"
+#include "AnimNextExecuteContext.h"
+
+#if WITH_EDITOR	
+#include "Engine/ExternalAssetDependencyGatherer.h"
+REGISTER_ASSETDEPENDENCY_GATHERER(FExternalAssetDependencyGatherer, UAnimNextSharedVariables);
+#endif // WITH_EDITOR
+
+UAnimNextSharedVariables::UAnimNextSharedVariables(const FObjectInitializer& ObjectInitializer)
+	: Super(ObjectInitializer)
+{
+	ExtendedExecuteContext.SetContextPublicDataStruct(FAnimNextExecuteContext::StaticStruct());
+}
diff --git a/Engine/Plugins/Experimental/AnimNext/Source/AnimNext/Private/Variables/AnimNextSoftVariableReference.cpp b/Engine/Plugins/Experimental/AnimNext/Source/AnimNext/Private/Variables/AnimNextSoftVariableReference.cpp
new file mode 100644
--- /dev/null
+++ b/Engine/Plugins/Experimental/AnimNext/Source/AnimNext/Private/Variables/AnimNextSoftVariableReference.cpp
@@ -0,0 +1,14 @@
+﻿// Copyright Epic Games, Inc. All Rights Reserved.
+
+#include "Variables/AnimNextSoftVariableReference.h"
+#include "AnimNextRigVMAsset.h"
+#include "Variables/AnimNextVariableReference.h"
+
+FAnimNextSoftVariableReference::FAnimNextSoftVariableReference(const FAnimNextVariableReference& InVariableReference)
+	: Name(InVariableReference.GetName())
+	, SoftObjectPath(InVariableReference.GetObject())
+{
+	const UObject* Object = InVariableReference.GetObject();
+	check(Object == nullptr || Object->IsA<UAnimNextRigVMAsset>() || Object->IsA<UScriptStruct>());
+}
+
diff --git a/Engine/Plugins/Experimental/AnimNext/Source/AnimNext/Private/Variables/RigVMTrait_OverrideVariables.h b/Engine/Plugins/Experimental/AnimNext/Source/AnimNext/Private/Variables/RigVMTrait_OverrideVariables.h
new file mode 100644
--- /dev/null
+++ b/Engine/Plugins/Experimental/AnimNext/Source/AnimNext/Private/Variables/RigVMTrait_OverrideVariables.h
@@ -0,0 +1,88 @@
+// Copyright Epic Games, Inc. All Rights Reserved.
+
+#pragma once
+
+#include "RigVMCore/RigVMTrait.h"
+#include "Param/ParamType.h"
+#include "RigVMTrait_OverrideVariables.generated.h"
+
+class UAnimNextRigVMAsset;
+
+namespace UE::AnimNext
+{
+	struct FVariableOverrides;
+}
+
+// Override info for a single variable
+USTRUCT(BlueprintInternalUseOnly)
+struct FAnimNextVariableOverrideInfo
+{
+	GENERATED_BODY()
+
+	FAnimNextVariableOverrideInfo() = default;
+
+	FAnimNextVariableOverrideInfo(FName InName, const FAnimNextParamType& InType, const FString& InDefaultValue)
+		: Name(InName)
+		, Type(InType)
+		, DefaultValue(InDefaultValue)
+	{}
+
+	UPROPERTY(EditAnywhere, Category = "Variables")
+	FName Name = "Variable";
+
+	UPROPERTY(EditAnywhere, Category = "Variables")
+	FAnimNextParamType Type = FAnimNextParamType(FAnimNextParamType::EValueType::Bool);
+
+	UPROPERTY()
+	FString DefaultValue;
+};
+
+// A trait that allows insertion of specific asset/struct variable overrides onto an override variables node
+USTRUCT(meta=(Hidden))
+struct ANIMNEXT_API FRigVMTrait_OverrideVariables : public FRigVMTrait
+{
+	GENERATED_BODY()
+
+	// FRigVMTrait interface
+#if WITH_EDITOR
+	virtual void GetProgrammaticPins(URigVMController* InController, int32 InParentPinIndex, const URigVMPin* InTraitPin, const FString& InDefaultValue, FRigVMPinInfoArray& OutPinArray) const override;
+	virtual bool ShouldCreatePinForProperty(const FProperty* InProperty) const override;
+#endif
+
+	// Override point for derived types to generate overrides
+	virtual void GenerateOverrides(const FRigVMTraitScope& InTraitScope, TArray<UE::AnimNext::FVariableOverrides>& OutOverrides) const {}
+
+	// Cached override data
+	UPROPERTY(EditAnywhere, Category = "Variables", meta=(Input, DetailsOnly, TitleProperty = "Name"))
+	TArray<FAnimNextVariableOverrideInfo> Overrides;
+};
+
+USTRUCT()
+struct ANIMNEXT_API FRigVMTrait_OverrideVariablesForAsset : public FRigVMTrait_OverrideVariables
+{
+	GENERATED_BODY()
+
+#if WITH_EDITOR
+	virtual FString GetDisplayName() const override;
+#endif
+	virtual void GenerateOverrides(const FRigVMTraitScope& InTraitScope, TArray<UE::AnimNext::FVariableOverrides>& OutOverrides) const override;
+
+	// The asset to use to generate overrides
+	UPROPERTY(EditAnywhere, Category = "Variables", meta=(Input, Constant, DetailsOnly))
+	TObjectPtr<const UAnimNextRigVMAsset> Asset;
+};
+
+USTRUCT()
+struct ANIMNEXT_API FRigVMTrait_OverrideVariablesForStruct : public FRigVMTrait_OverrideVariables
+{
+	GENERATED_BODY()
+
+#if WITH_EDITOR
+	virtual FString GetDisplayName() const override;
+#endif
+	virtual void GenerateOverrides(const FRigVMTraitScope& InTraitScope, TArray<UE::AnimNext::FVariableOverrides>& OutOverrides) const override;
+
+	// The asset to use to generate overrides
+	UPROPERTY(EditAnywhere, Category = "Variables", meta=(Input, Constant, DetailsOnly))
+	TObjectPtr<const UScriptStruct> Struct;
+};
diff --git a/Engine/Plugins/Experimental/AnimNext/Source/AnimNext/Private/Variables/VariableOverrides.cpp b/Engine/Plugins/Experimental/AnimNext/Source/AnimNext/Private/Variables/VariableOverrides.cpp
new file mode 100644
--- /dev/null
+++ b/Engine/Plugins/Experimental/AnimNext/Source/AnimNext/Private/Variables/VariableOverrides.cpp
@@ -0,0 +1,25 @@
+﻿// Copyright Epic Games, Inc. All Rights Reserved.
+
+#include "VariableOverrides.h"
+
+#include "AnimNextRigVMAsset.h"
+
+namespace UE::AnimNext
+{
+
+FVariableOverrides::FVariableOverrides(const UAnimNextRigVMAsset* InAsset, TArray<FOverride>&& InOverrides)
+	: Overrides(MoveTemp(InOverrides))
+{
+	check(InAsset);
+
+	AssetOrStructData.Set<FAssetType>(InAsset);
+}
+
+FVariableOverrides::FVariableOverrides(const UScriptStruct* InStruct, TArray<FOverride>&& InOverrides)
+	: Overrides(MoveTemp(InOverrides))
+{
+	AssetOrStructData.Set<FStructType>(InStruct);
+
+}
+
+}
diff --git a/Engine/Plugins/Experimental/AnimNext/Source/AnimNext/Public/AnimNextFunctionHandle.h b/Engine/Plugins/Experimental/AnimNext/Source/AnimNext/Public/AnimNextFunctionHandle.h
new file mode 100644
--- /dev/null
+++ b/Engine/Plugins/Experimental/AnimNext/Source/AnimNext/Public/AnimNextFunctionHandle.h
@@ -0,0 +1,38 @@
+﻿// Copyright Epic Games, Inc. All Rights Reserved.
+
+#pragma once
+
+#include "CoreTypes.h"
+
+class UAnimNextRigVMAsset;
+
+namespace UE::AnimNext
+{
+
+// Opaque handle that can be used to call RigVM functions on UAnimNextRigVMAsset
+struct FFunctionHandle
+{
+	// Check to see if this is a valid handle
+	bool IsValid() const
+	{
+		return FunctionIndex != INDEX_NONE;
+	}
+
+#if WITH_EDITOR
+	// Check to see if this is a valid handle and that it can be used for the specified VM (hash comparison)
+	bool IsValidForVM(uint32 InVMHash) const
+	{
+		return IsValid() && VMHash == InVMHash;
+	}
+#endif
+
+private:
+	friend ::UAnimNextRigVMAsset;
+
+	int32 FunctionIndex = INDEX_NONE;
+#if WITH_EDITOR
+	uint32 VMHash = 0;
+#endif
+};
+
+}
\ No newline at end of file
diff --git a/Engine/Plugins/Experimental/AnimNext/Source/AnimNext/Public/AnimNextRigVMAsset.h b/Engine/Plugins/Experimental/AnimNext/Source/AnimNext/Public/AnimNextRigVMAsset.h
--- a/Engine/Plugins/Experimental/AnimNext/Source/AnimNext/Public/AnimNextRigVMAsset.h
+++ b/Engine/Plugins/Experimental/AnimNext/Source/AnimNext/Public/AnimNextRigVMAsset.h
@@ -5,18 +5,30 @@
 #include "CoreMinimal.h"
 #include "RigVMCore/RigVMExecuteContext.h"
 #include "RigVMHost.h"
+#include "Param/ParamType.h"
+#include "Variables/AnimNextVariableReference.h"
 #include "AnimNextRigVMAsset.generated.h"
 
 #define UE_API ANIMNEXT_API
 
+class UAnimNextRigVMAsset;
 class UAnimNextVariableEntry;
 class UAnimNextRigVMAssetEditorData;
-struct FAnimNextModuleInstance;
+class UAnimNextAnimGraphSettings;
+struct FUAFAssetInstance;
+struct FUAFInstanceVariableContainer;
+struct FUAFInstanceVariableData;
+struct FUAFInstanceVariableDataProxy;
+struct FAnimNextRigVMFunctionData;
+struct FUAFRigVMComponent;
 
 namespace UE::AnimNext
 {
-	struct FProxyVariablesContext;
+	struct FAnimGraphBuilderContext;
+	struct FFunctionHandle;
 	struct FModuleEventTickFunction;
+	struct FVariableOverrides;
+	struct FInjectionInfo;
 }
 
 namespace UE::AnimNext::UncookedOnly
@@ -33,10 +45,13 @@ class UAnimNextRigVMAsset : public URigVMHost
 protected:
 	friend class UAnimNextVariableEntry;
 	friend class UAnimNextRigVMAssetEditorData;
-	friend struct FAnimNextModuleInstance;
+	friend UAnimNextAnimGraphSettings;
+	friend FUAFRigVMComponent;
 	friend struct UE::AnimNext::FModuleEventTickFunction;
-	friend struct UE::AnimNext::FProxyVariablesContext;
 	friend struct UE::AnimNext::UncookedOnly::FUtils;
+	friend FUAFInstanceVariableData;
+	friend UE::AnimNext::FInjectionInfo;
+	friend UE::AnimNext::FAnimGraphBuilderContext;
 
 	UE_API UAnimNextRigVMAsset(const FObjectInitializer& ObjectInitializer);
 
@@ -47,16 +62,28 @@ protected:
 	UE_API virtual void PreDuplicate(FObjectDuplicationParameters& DupParams) override;
 
 public:
-	// Get public variable defaults property bag
-	const FInstancedPropertyBag& GetPublicVariableDefaults() const { return PublicVariableDefaults; }
+	// Get variable defaults property bag
+	const FInstancedPropertyBag& GetVariableDefaults() const { return VariableDefaults; }
 
 	// Get the RigVM object we host
 	const URigVM* GetRigVM() const { return RigVM; }
 
+	// Call a function handle using a property bag as arguments.
+	UE_API void CallFunctionHandle(UE::AnimNext::FFunctionHandle InHandle, FRigVMExtendedExecuteContext& InContext, FInstancedPropertyBag& InArgs);
+
+	// Get a function handle given its mangled event name
+	UE_API UE::AnimNext::FFunctionHandle GetFunctionHandle(FName InEventName) const;
+
 private:
 	// URigVMHost interface
 	UE_API virtual TArray<FRigVMExternalVariable> GetExternalVariablesImpl(bool bFallbackToBlueprint) const;
 
+	// Get all the assets whose variables we will reference/instance at runtime
+	TConstArrayView<TObjectPtr<const UAnimNextRigVMAsset>> GetReferencedVariableAssets() const { return ReferencedVariableAssets; }
+
+	// Get all the native structs whose variables we will reference/instance at runtime
+	TConstArrayView<TObjectPtr<const UScriptStruct>> GetReferencedVariableStructs() const { return ReferencedVariableStructs; }
+
 protected:
 	// The ExtendedExecuteContext object holds the common work data used by the RigVM internals. It is populated during the initial VM initialization.
 	// Each instance of an AnimGraph requires a copy of this context and a call to initialize the VM instance with the context copy, 
@@ -72,9 +99,26 @@ protected:
 	UPROPERTY()
 	FInstancedPropertyBag VariableDefaults;
 
-	// Public variables (for easy duplication)
+	// Variables property bag (including all variables from referenced assets, as external variables need stable property ptrs). 
+	// Default values are not valid, this is just stored here for FProperty stability
+	UPROPERTY()
+	FInstancedPropertyBag CombinedPropertyBag;
+
+	// All the assets whose variables we will reference/instance at runtime
+	UPROPERTY()
+	TArray<TObjectPtr<const UAnimNextRigVMAsset>> ReferencedVariableAssets;
+
+	// All the native structs whose variables we will reference/instance at runtime
+	UPROPERTY()
+	TArray<TObjectPtr<const UScriptStruct>> ReferencedVariableStructs;
+
+	// All the natively-callable RigVM functions
+	UPROPERTY()
+	TArray<FAnimNextRigVMFunctionData> FunctionData;
+
+	// Reference to the default injection site
 	UPROPERTY()
-	FInstancedPropertyBag PublicVariableDefaults;
+	FAnimNextVariableReference DefaultInjectionSite;
 
 #if WITH_EDITORONLY_DATA
 	UPROPERTY(VisibleAnywhere, Instanced, Category = "Editor Data", meta = (ShowInnerProperties))
diff --git a/Engine/Plugins/Experimental/AnimNext/Source/AnimNext/Public/Component/AnimNextComponent.h b/Engine/Plugins/Experimental/AnimNext/Source/AnimNext/Public/Component/AnimNextComponent.h
--- a/Engine/Plugins/Experimental/AnimNext/Source/AnimNext/Public/Component/AnimNextComponent.h
+++ b/Engine/Plugins/Experimental/AnimNext/Source/AnimNext/Public/Component/AnimNextComponent.h
@@ -2,12 +2,10 @@
 
 #pragma once
 
-#include "AnimNextPublicVariablesProxy.h"
 #include "Components/ActorComponent.h"
 #include "Module/ModuleHandle.h"
 #include "Param/ParamType.h"
 #include "TraitCore/TraitEvent.h"
-#include "Variables/IAnimNextVariableProxyHost.h"
 #include "Module/AnimNextModuleInitMethod.h"
 #include "Module/ModuleTaskContext.h"
 #include "Module/TaskRunLocation.h"
@@ -17,11 +15,7 @@
 struct FAnimNextComponentInstanceData;
 class UAnimNextComponentWorldSubsystem;
 class UAnimNextModule;
-
-namespace UE::AnimNext
-{
-	struct FProxyVariablesContext;
-}
+struct AnimNextVariableReference;
 
 namespace UE::AnimNext::UncookedOnly
 {
@@ -29,7 +23,7 @@ namespace UE::AnimNext::UncookedOnly
 }
 
 UCLASS(MinimalAPI, meta = (BlueprintSpawnableComponent))
-class UAnimNextComponent : public UActorComponent, public IAnimNextVariableProxyHost
+class UAnimNextComponent : public UActorComponent
 {
 	GENERATED_BODY()
 
@@ -39,148 +33,43 @@ class UAnimNextComponent : public UActorComponent, public IAnimNextVariableProxy
 	virtual void BeginPlay() override;
 	virtual void EndPlay(const EEndPlayReason::Type EndPlayReason) override;
 
-	// IAnimNextVariableProxyHost interface
-	virtual void FlipPublicVariablesProxy(const UE::AnimNext::FProxyVariablesContext& InContext) override;
-
-#if WITH_EDITOR
-	// Called back to refresh any cached data on module compilation
-	void OnModuleCompiled();
-#endif
-
-	// (Re-)create the public variable proxy
-	void CreatePublicVariablesProxy();
-
-	// Destroy the public variable proxy
-	void DestroyPublicVariablesProxy();
-
 public:
 	// Sets a module variable's value.
 	// @param    Name     The name of the variable to set
 	// @param    Value    The value to set the variable to
-	UFUNCTION(BlueprintCallable, Category = "UAF", CustomThunk, meta = (CustomStructureParam = Value, UnsafeDuringActorConstruction))
-	ANIMNEXT_API void SetVariable(UPARAM(meta = (CustomWidget = "VariableName")) FName Name, int32 Value);
-
-	/**
-	 * Module variable value getters.
-	 * Numeric types (bool, (u)int32, (u)int64, float, double) support type conversion.
-	 * Struct & Object types will be const, if you need to modify them use SetVariable[Struct, Object]Ref
-	 * @param    Name     The name of the variable to get
-	 * @return            The variable value on success, None otherwise
-	 */
-	ANIMNEXT_API TOptional<bool> GetVariableBool(const FName Name) const;
-	ANIMNEXT_API TOptional<uint8> GetVariableByte(const FName Name) const;
-	ANIMNEXT_API TOptional<int32> GetVariableInt32(const FName Name) const;
-	ANIMNEXT_API TOptional<uint32> GetVariableUInt32(const FName Name) const;
-	ANIMNEXT_API TOptional<int64> GetVariableInt64(const FName Name) const;
-	ANIMNEXT_API TOptional<uint64> GetVariableUInt64(const FName Name) const; 
-	ANIMNEXT_API TOptional<float> GetVariableFloat(const FName Name) const;
-	ANIMNEXT_API TOptional<double> GetVariableDouble(const FName Name) const;
-	ANIMNEXT_API TOptional<FName> GetVariableName(const FName Name) const;
-	ANIMNEXT_API TOptional<FString> GetVariableString(const FName Name) const;
-	ANIMNEXT_API TOptional<uint8> GetVariableEnum(const FName Name, const UEnum* RequestedEnum) const;
-	ANIMNEXT_API TOptional<FConstStructView> GetVariableStruct(const FName Name, const UScriptStruct* RequestedStruct = nullptr) const;
-	ANIMNEXT_API TOptional<const UObject*> GetVariableObject(const FName Name, const UClass* RequestedClass = nullptr) const;
-	ANIMNEXT_API TOptional<const UClass*> GetVariableClass(const FName Name) const;
-	ANIMNEXT_API TOptional<FSoftObjectPath> GetVariableSoftPath(const FName Name) const;
-
-	/** @return enum value of specified type. */
-	template <typename T>
-	TOptional<T> GetVariableEnum(const FName Name) const
-	{
-		static_assert(TIsEnum<T>::Value, "Should only call this with enum types");
+	UE_DEPRECATED(5.6, "This function is no longer used, please use SetVariableReference (for Blueprint) or SetVariable that takes a FAnimNextVariableReference")
+	UFUNCTION(BlueprintCallable, Category = "UAF", DisplayName = "Set Variable", CustomThunk, meta = (CustomStructureParam = Value, UnsafeDuringActorConstruction, DeprecatedFunction, DeprecationMessage = "This function has been deprecated, please use Set Variable that takes a variable reference instead"))
+	ANIMNEXT_API void BlueprintSetVariable(FName Name, int32 Value);
 
-		TOptional<uint8> Result = GetVariableEnum(Name, StaticEnum<T>());
-		if (!Result.IsSet())
-		{
-			return TOptional<T>();
-		}
-		return TOptional<T>(static_cast<T>(Result.GetValue()));
-	}
-
-	/** @return struct reference of specified type. */
-	template <typename T>
-	TOptional<const T*> GetVariableStruct(const FName Name) const
+	// Sets a module variable's value.
+	// @param    Name     The variable to set
+	// @param    Value    The value to set the variable to
+	UFUNCTION(BlueprintCallable, BlueprintInternalUseOnly, Category = "UAF", DisplayName = "Set Variable", CustomThunk, meta = (CustomStructureParam = Value, AutoCreateRefTerm = "Variable, Value", UnsafeDuringActorConstruction))
+	ANIMNEXT_API void BlueprintSetVariableReference(const FAnimNextVariableReference& Variable, const int32& Value);
+
+	// Set a variable.
+	// @param    InVariable         The variable to set
+	// @param    InValue            The value to set the variable to
+	// @return true if the variable was set correctly
+	template<typename ValueType>
+	bool SetVariable(const FAnimNextVariableReference& InVariable, const ValueType& InValue)
 	{
-		TOptional<FConstStructView> Result = GetVariableStruct(Name, TBaseStructure<T>::Get());
-		if (!Result.IsSet())
-		{
-			return TOptional<const T*>();
-		}
-		if (const T* ValuePtr = Result.GetValue().GetPtr<const T>())
-		{
-			return TOptional<const T*>(ValuePtr);
-		}
-		return TOptional<const T*>();
+		return SetVariableInternal(InVariable, FAnimNextParamType::GetType<ValueType>(), TConstArrayView<uint8>(reinterpret_cast<const uint8*>(&InValue), sizeof(ValueType))) == EPropertyBagResult::Success;
 	}
 
-	/** @return object pointer value of specified type. */
-	template <typename T>
-	TOptional<const T*> GetVariableObject(const FName Name) const
+	// Accesses a variable for writing.
+	// General read access is not provided via this API due to the current double-buffering strategy used to communicate variable writes to worker threads.
+	// This is intended to allow for copy-free writing of larger data structures & arrays, rather than read access. 
+	// @param    InVariable         The variable to set
+	// @param    InFunction         The function used to modify the variable, called back immediately on the game thread if the variable exists.
+	// @return true if the variable was accessed successfully
+	template<typename ValueType>
+	bool WriteVariable(const FAnimNextVariableReference& InVariable, TFunctionRef<void(ValueType&)> InFunction)
 	{
-		static_assert(TIsDerivedFrom<T, UObject>::Value, "Should only call this with object types");
-
-		TOptional<const UObject*> Result = GetVariableObject(Name, T::StaticClass());
-		if (!Result.IsSet())
+		return WriteVariableInternal(InVariable, FAnimNextParamType::GetType<ValueType>(), [&InFunction](TArrayView<uint8> InData)
 		{
-			return TOptional<const T*>();
-		}
-		if (Result.GetValue() == nullptr)
-		{
-			return TOptional<const T*>(nullptr);
-		}
-		if (const T* Object = Cast<const T>(Result.GetValue()))
-		{
-			return TOptional<const T*>(Object);
-		}
-		return TOptional<const T*>();
-	}
-
-	/**
-	 * Module variable value setters.
-	 * @param    Name     The name of the variable to set
-	 * @param    Value    The value to set the variable to
-	 * @return            True on success, false otherwise
-	 */
-	ANIMNEXT_API bool SetVariableBool(const FName Name, const bool bInValue);
-	ANIMNEXT_API bool SetVariableByte(const FName Name, const uint8 InValue);
-	ANIMNEXT_API bool SetVariableInt32(const FName Name, const int32 InValue);
-	ANIMNEXT_API bool SetVariableUInt32(const FName Name, const uint32 InValue);
-	ANIMNEXT_API bool SetVariableInt64(const FName Name, const int64 InValue);
-	ANIMNEXT_API bool SetVariableUInt64(const FName Name, const uint64 InValue);
-	ANIMNEXT_API bool SetVariableFloat(const FName Name, const float InValue);
-	ANIMNEXT_API bool SetVariableDouble(const FName Name, const double InValue);
-	ANIMNEXT_API bool SetVariableName(const FName Name, const FName InValue);
-	ANIMNEXT_API bool SetVariableString(const FName Name, const FString& InValue);
-	ANIMNEXT_API bool SetVariableEnum(const FName Name, const uint8 InValue, const UEnum* Enum);
-	ANIMNEXT_API bool SetVariableStruct(const FName Name, FConstStructView InValue);
-	ANIMNEXT_API bool SetVariableStructRef(const FName Name, TFunctionRef<void(FStructView)> InStructRefSetter, const UScriptStruct* RequestedStruct = nullptr);
-	ANIMNEXT_API bool SetVariableObject(const FName Name, UObject* InValue);
-	ANIMNEXT_API bool SetVariableObjectRef(const FName Name, TFunctionRef<void(UObject*)> InObjectRefSetter, const UClass* RequestedClass = nullptr);
-	ANIMNEXT_API bool SetVariableClass(const FName Name, UClass* InValue);
-	ANIMNEXT_API bool SetVariableSoftPath(const FName Name, const FSoftObjectPath& InValue);
-	ANIMNEXT_API bool SetVariableArrayRef(const FName Name, TFunctionRef<void(FPropertyBagArrayRef&)> InArrayRefSetter);
-
-	/** Sets enum variable with specified type. */
-	template <typename T>
-	bool SetVariableEnum(const FName Name, const T InValue)
-	{
-		static_assert(TIsEnum<T>::Value, "Should only call this with enum types");
-		return SetVariableEnum(Name, static_cast<uint8>(InValue), StaticEnum<T>());
-	}
-
-	/** Sets struct variable with specified type. */
-	template <typename T>
-	bool SetVariableStruct(const FName Name, const T& InValue)
-	{
-		return SetVariableStruct(Name, FConstStructView::Make(InValue));
-	}
-
-	/** Sets object pointer variable with specified type. */
-	template <typename T>
-	bool SetVariableObject(const FName Name, T* InValue)
-	{
-		static_assert(TIsDerivedFrom<T, UObject>::Value, "Should only call this with object types");
-		return SetVariableObject(Name, (UObject*)InValue);
+			InFunction(*reinterpret_cast<ValueType*>(InData.GetData()));
+		}) == EPropertyBagResult::Success;
 	}
 
 	// Whether this component is currently updating
@@ -291,7 +180,14 @@ public:
 	ANIMNEXT_API FAnimNextModuleHandle BlueprintGetModuleHandle() const;
 
 private:
-	DECLARE_FUNCTION(execSetVariable);
+	DECLARE_FUNCTION(execBlueprintSetVariable);
+	DECLARE_FUNCTION(execBlueprintSetVariableReference);
+
+	// Sets the value of the specified variable
+	ANIMNEXT_API EPropertyBagResult SetVariableInternal(const FAnimNextVariableReference& InVariable, const FAnimNextParamType& InType, TConstArrayView<uint8> InNewValue);
+
+	// Accesses value of the specified variable for writing
+	ANIMNEXT_API EPropertyBagResult WriteVariableInternal(const FAnimNextVariableReference& InVariable, const FAnimNextParamType& InType, TFunctionRef<void(TArrayView<uint8>)> InFunction);
 
 private:
 	friend struct UE::AnimNext::UncookedOnly::FUtils;
@@ -307,16 +203,6 @@ private:
 	// Handle to the registered module
 	UE::AnimNext::FModuleHandle ModuleHandle;
 
-	// Lock for public variables proxy
-	mutable FRWLock PublicVariablesLock;
-
-	// Proxy public variables
-	UPROPERTY()
-	FAnimNextPublicVariablesProxy PublicVariablesProxy;
-
-	// Map from name->proxy variable index
-	TMap<FName, int32> PublicVariablesProxyMap;
-
 	// How to initialize the module
 	UPROPERTY(EditAnywhere, Category="Module")
 	EAnimNextModuleInitMethod InitMethod = EAnimNextModuleInitMethod::InitializeAndPauseInEditor;
diff --git a/Engine/Plugins/Experimental/AnimNext/Source/AnimNext/Public/Module/ModuleTaskContext.h b/Engine/Plugins/Experimental/AnimNext/Source/AnimNext/Public/Module/ModuleTaskContext.h
--- a/Engine/Plugins/Experimental/AnimNext/Source/AnimNext/Public/Module/ModuleTaskContext.h
+++ b/Engine/Plugins/Experimental/AnimNext/Source/AnimNext/Public/Module/ModuleTaskContext.h
@@ -3,10 +3,11 @@
 #pragma once
 
 #include "CoreMinimal.h"
+#include "InstanceTaskContext.h"
 
 #include "TraitCore/TraitEvent.h"
 
-struct FAnimNextModuleInstanceComponent;
+struct FUAFModuleInstanceComponent;
 struct FAnimNextModuleInstance;
 struct FInstancedPropertyBag;
 struct FAnimNextModuleInjectionComponent;
@@ -22,7 +23,7 @@ namespace UE::AnimNext
 {
 
 // Context passed to schedule task callbacks
-struct FModuleTaskContext
+struct FModuleTaskContext : public FInstanceTaskContext
 {
 public:
 	// Queues an input trait event
@@ -30,12 +31,12 @@ public:
 	ANIMNEXT_API void QueueInputTraitEvent(FAnimNextTraitEventPtr Event) const;
 
 	// Access a module instance component of the specified type. If the component exists, then InFunction will be called
-	ANIMNEXT_API void TryAccessComponent(UScriptStruct* InComponentType, TFunctionRef<void(FAnimNextModuleInstanceComponent&)> InFunction) const;
+	ANIMNEXT_API void TryAccessComponent(UScriptStruct* InComponentType, TFunctionRef<void(FUAFModuleInstanceComponent&)> InFunction) const;
 
 	template<typename ComponentType>
 	void TryAccessComponent(TFunctionRef<void(ComponentType&)> InFunction) const
 	{
-		TryAccessComponent(ComponentType::StaticStruct(), [&InFunction](FAnimNextModuleInstanceComponent& InComponent)
+		TryAccessComponent(TBaseStructure<ComponentType>::Get(), [&InFunction](FUAFModuleInstanceComponent& InComponent)
 		{
 			InFunction(static_cast<ComponentType&>(InComponent));
 		});
@@ -47,7 +48,7 @@ private:
 	FModuleTaskContext(FAnimNextModuleInstance& InModuleInstance);
 
 	// The module instance currently running
-	FAnimNextModuleInstance* ModuleInstance;
+	FAnimNextModuleInstance& ModuleInstance;
 
 	friend UE::AnimNext::FModuleEventTickFunction;
 	friend ::FAnimNextModuleInjectionComponent;
diff --git a/Engine/Plugins/Experimental/AnimNext/Source/AnimNext/Public/Module/UAFModuleInstanceComponent.h b/Engine/Plugins/Experimental/AnimNext/Source/AnimNext/Public/Module/UAFModuleInstanceComponent.h
new file mode 100644
--- /dev/null
+++ b/Engine/Plugins/Experimental/AnimNext/Source/AnimNext/Public/Module/UAFModuleInstanceComponent.h
@@ -0,0 +1,45 @@
+// Copyright Epic Games, Inc. All Rights Reserved.
+
+#pragma once
+
+#include "CoreMinimal.h"
+#include "UAFAssetInstanceComponent.h"
+#include "UAFModuleInstanceComponent.generated.h"
+
+struct FAnimNextModuleInstance;
+struct FAnimNextTraitEvent;
+
+#define UE_API ANIMNEXT_API
+
+/** A module instance component is attached and owned by a module instance. */
+USTRUCT()
+struct FUAFModuleInstanceComponent : public FUAFAssetInstanceComponent
+{
+	GENERATED_BODY()
+
+	using ContainerType = FAnimNextModuleInstance;
+
+	FUAFModuleInstanceComponent() = default;
+
+	virtual ~FUAFModuleInstanceComponent() override = default;
+
+	// Returns the owning module instance this component lives on
+	UE_API FAnimNextModuleInstance* GetModuleInstancePtr();
+
+	// Returns the owning module instance this component lives on
+	UE_API FAnimNextModuleInstance& GetModuleInstance();
+
+	// Returns the owning module instance this component lives on
+	UE_API const FAnimNextModuleInstance& GetModuleInstance() const;
+
+	// Called during module execution for any events to be handled
+	virtual void OnTraitEvent(FAnimNextTraitEvent& Event) {}
+
+	// Called at end of module execution each frame
+	virtual void OnEndExecution(float InDeltaTime) {}
+
+private:
+	friend struct FAnimNextModuleInstance;
+};
+
+#undef UE_API
\ No newline at end of file
diff --git a/Engine/Plugins/Experimental/AnimNext/Source/AnimNext/Public/Param/ParamType.h b/Engine/Plugins/Experimental/AnimNext/Source/AnimNext/Public/Param/ParamType.h
--- a/Engine/Plugins/Experimental/AnimNext/Source/AnimNext/Public/Param/ParamType.h
+++ b/Engine/Plugins/Experimental/AnimNext/Source/AnimNext/Public/Param/ParamType.h
@@ -69,7 +69,12 @@ public:
 	FAnimNextParamType() = default;
 
 	/** Construct a parameter type from the passed in value, container and object type. */
-	UE_API FAnimNextParamType(EValueType InValueType, EContainerType InContainerType = EContainerType::None, const UObject* InValueTypeObject = nullptr);
+	FAnimNextParamType(EValueType InValueType, EContainerType InContainerType = EContainerType::None, const UObject* InValueTypeObject = nullptr)
+		: ValueTypeObject(InValueTypeObject)
+		, ValueType(InValueType)
+		, ContainerType(InContainerType)
+	{
+	}
 
 private:
 	/** Pointer to object that defines the Enum, Struct, or Class. */
@@ -87,113 +92,104 @@ private:
 private:
 	/** Helper function for GetType */
 	template<typename ParamType>
-	static constexpr void GetTypeInner(FAnimNextParamType& ParameterType)
+	static FAnimNextParamType GetTypeInner(EContainerType InContainerType)
 	{
 		using NonPtrParamType = std::remove_pointer_t<ParamType>;
 
 		if constexpr (std::is_same_v<ParamType, bool>)
 		{
-			ParameterType.ValueType = EValueType::Bool;
+			return FAnimNextParamType(EValueType::Bool, InContainerType);
 		}
 		else if constexpr (std::is_same_v<ParamType, uint8>)
 		{
-			ParameterType.ValueType = EValueType::Byte;
+			return FAnimNextParamType(EValueType::Byte, InContainerType);
 		}
 		else if constexpr (std::is_same_v<ParamType, int32> || std::is_same_v<ParamType, int>)
 		{
-			ParameterType.ValueType = EValueType::Int32;
+			return FAnimNextParamType(EValueType::Int32, InContainerType);
 		}
 		else if constexpr (std::is_same_v<ParamType, int64>)
 		{
-			ParameterType.ValueType = EValueType::Int64;
+			return FAnimNextParamType(EValueType::Int64, InContainerType);
 		}
 		else if constexpr (std::is_same_v<ParamType, float>)
 		{
-			ParameterType.ValueType = EValueType::Float;
+			return FAnimNextParamType(EValueType::Float, InContainerType);
 		}
 		else if constexpr (std::is_same_v<ParamType, double>)
 		{
-			ParameterType.ValueType = EValueType::Double;
+			return FAnimNextParamType(EValueType::Double, InContainerType);
 		}
 		else if constexpr (std::is_same_v<ParamType, FName>)
 		{
-			ParameterType.ValueType = EValueType::Name;
+			return FAnimNextParamType(EValueType::Name, InContainerType);
 		}
 		else if constexpr (std::is_same_v<ParamType, FString>)
 		{
-			ParameterType.ValueType = EValueType::String;
+			return FAnimNextParamType(EValueType::String, InContainerType);
 		}
 		else if constexpr (std::is_same_v<ParamType, FText>)
 		{
-			ParameterType.ValueType = EValueType::Text;
+			return FAnimNextParamType(EValueType::Text, InContainerType);
 		}
 		else if constexpr (TIsUEnumClass<ParamType>::Value)
 		{
-			ParameterType.ValueType = EValueType::Enum;
-			ParameterType.ValueTypeObject = StaticEnum<ParamType>();
+			return FAnimNextParamType(EValueType::Enum, InContainerType, StaticEnum<ParamType>());
 		}
 		else if constexpr (TModels<CStaticStructProvider, ParamType>::Value)
 		{
-			ParameterType.ValueType = EValueType::Struct;
-			ParameterType.ValueTypeObject = ParamType::StaticStruct();
+			return FAnimNextParamType(EValueType::Struct, InContainerType, ParamType::StaticStruct());
 		}
 		else if constexpr (TModels<CBaseStructureProvider, ParamType>::Value)
 		{
-			ParameterType.ValueType = EValueType::Struct;
-			ParameterType.ValueTypeObject = TBaseStructure<ParamType>::Get();
+			return FAnimNextParamType(EValueType::Struct, InContainerType, TBaseStructure<ParamType>::Get());
 		}
 		else if constexpr (TModels<CStaticClassProvider, NonPtrParamType>::Value)
 		{
 			if constexpr (std::is_same_v<NonPtrParamType, UClass>)
 			{
-				ParameterType.ValueType = EValueType::Class;
-				ParameterType.ValueTypeObject = NonPtrParamType::StaticClass();
+				return FAnimNextParamType(EValueType::Class, InContainerType, UObject::StaticClass());
 			}
 			else
 			{
-				ParameterType.ValueType = EValueType::Object;
-				ParameterType.ValueTypeObject = NonPtrParamType::StaticClass();
+				return FAnimNextParamType(EValueType::Object, InContainerType, NonPtrParamType::StaticClass());
 			}
 		}
 		else if constexpr (TIsTObjectPtr<ParamType>::Value)
 		{
 			if constexpr (std::is_same_v<ParamType, TObjectPtr<UClass>>)
 			{
-				ParameterType.ValueType = EValueType::Class;
-				ParameterType.ValueTypeObject = UObject::StaticClass();
+				return FAnimNextParamType(EValueType::Class, InContainerType, UObject::StaticClass());
 			}
 			else
 			{
-				ParameterType.ValueType = EValueType::Object;
-				ParameterType.ValueTypeObject = ParamType::ElementType::StaticClass();
+				return FAnimNextParamType(EValueType::Object, InContainerType, ParamType::ElementType::StaticClass());
 			}
 		}
 		else if constexpr (TIsTSubclassOf<ParamType>::Value)
 		{
-			ParameterType.ValueType = EValueType::Class;
-			ParameterType.ValueTypeObject = ParamType::ElementType::StaticClass();
+			return FAnimNextParamType(EValueType::Class, InContainerType, ParamType::ElementType::StaticClass());
 		}
 		else if constexpr (UE::AnimNext::Private::TIsSoftObjectPtr<ParamType>::Value)
 		{
-			ParameterType.ValueType = EValueType::SoftObject;
-			ParameterType.ValueTypeObject = ParamType::ElementType::StaticClass();
+			return FAnimNextParamType(EValueType::SoftObject, InContainerType, ParamType::ElementType::StaticClass());
 		}
 		else if constexpr (UE::AnimNext::Private::TIsSoftClassPtr<ParamType>::Value)
 		{
-			ParameterType.ValueType = EValueType::SoftClass;
-			ParameterType.ValueTypeObject = ParamType::ElementType::StaticClass();
+			return FAnimNextParamType(EValueType::SoftClass, InContainerType, ParamType::ElementType::StaticClass());
 		}
 		else if constexpr (std::is_same_v<ParamType, uint32>)
 		{
-			ParameterType.ValueType = EValueType::UInt32;
+			return FAnimNextParamType(EValueType::UInt32, InContainerType);
 		}
 		else if constexpr (std::is_same_v<ParamType, uint64>)
 		{
-			ParameterType.ValueType = EValueType::UInt64;
+			return FAnimNextParamType(EValueType::UInt64, InContainerType);
 		}
 		else
 		{
 			static_assert(sizeof(ParamType) == 0, "Type is not expressible as a FAnimNextParamType for available types.");
+			return FAnimNextParamType();
 		}
 	}
 
@@ -207,19 +203,14 @@ public:
 	{
 		using NonConstType = std::remove_const_t<ParamType>;
 
-		FAnimNextParamType ParameterType;
-
 		if constexpr (TIsTArray<NonConstType>::Value)
 		{
-			GetTypeInner<typename NonConstType::ElementType>(ParameterType);
-			ParameterType.ContainerType = EContainerType::Array;
+			return GetTypeInner<typename NonConstType::ElementType>(EContainerType::Array);
 		}
 		else
 		{
-			GetTypeInner<NonConstType>(ParameterType);
+			return GetTypeInner<NonConstType>(EContainerType::None);
 		}
-
-		return ParameterType;
 	}
 
 	/** Get the pointer to the object that defines the Enum, Struct, or Class. */
@@ -288,6 +279,13 @@ public:
 	/** Construct a parameter type from the passed in FProperty. */
 	static UE_API FAnimNextParamType FromProperty(const FProperty* InProperty);
 
+	/** Construct a parameter type from the passed in FProperty. */
+	static FAnimNextParamType FromPropertyBagPropertyDesc(const FPropertyBagPropertyDesc& InPropertyDesc)
+	{
+		check(InPropertyDesc.ContainerTypes.Num() <= 1);	// Cannot create a param type from a nested container descriptor
+		return FAnimNextParamType(InPropertyDesc.ValueType, InPropertyDesc.ContainerTypes.GetFirstContainerType(), InPropertyDesc.ValueTypeObject);
+	}
+
 	/** Equality operator */
 	friend bool operator==(const FAnimNextParamType& InLHS, const FAnimNextParamType& InRHS)
 	{
diff --git a/Engine/Plugins/Experimental/AnimNext/Source/AnimNext/Public/RewindDebugger/AnimNextTrace.h b/Engine/Plugins/Experimental/AnimNext/Source/AnimNext/Public/RewindDebugger/AnimNextTrace.h
--- a/Engine/Plugins/Experimental/AnimNext/Source/AnimNext/Public/RewindDebugger/AnimNextTrace.h
+++ b/Engine/Plugins/Experimental/AnimNext/Source/AnimNext/Public/RewindDebugger/AnimNextTrace.h
@@ -12,6 +12,7 @@
 #include "Misc/Guid.h"
 #include "ObjectTraceDefines.h"
 #include "Trace/Trace.h"
+#include "Templates/SharedPointer.h"
 
 #define ANIMNEXT_TRACE_ENABLED (OBJECT_TRACE_ENABLED && !(UE_BUILD_SHIPPING || UE_BUILD_TEST))
 
@@ -21,8 +22,9 @@ UE_TRACE_CHANNEL_EXTERN(AnimNextChannel, ANIMNEXT_API);
 
 class UObject;
 struct FAnimNextModuleInstance;
-struct FAnimNextDataInterfaceInstance;
+struct FUAFAssetInstance;
 struct FAnimNextGraphInstance;
+struct FUAFInstanceVariableContainer;
 
 namespace UE::AnimNext
 {
@@ -32,8 +34,11 @@ struct FAnimNextTrace
 	
 	ANIMNEXT_API static void Reset();
 	
-	ANIMNEXT_API static void OutputAnimNextInstance(const FAnimNextDataInterfaceInstance* DataInterface, const UObject* OuterObject);
-	ANIMNEXT_API static void OutputAnimNextVariables(const FAnimNextDataInterfaceInstance* DataInterface, const UObject* OuterObject);
+	ANIMNEXT_API static void OutputAnimNextInstance(const FUAFAssetInstance* DataInterface, const UObject* OuterObject);
+	ANIMNEXT_API static void OutputAnimNextVariables(const FUAFAssetInstance* DataInterface, const UObject* OuterObject);
+
+private:
+	ANIMNEXT_API static void OutputAnimNextVariableSet(const TSharedRef<FUAFInstanceVariableContainer>& VariableSet, uint64 InstanceId, const UObject* OuterObject);
 };
 
 }
diff --git a/Engine/Plugins/Experimental/AnimNext/Source/AnimNext/Public/Variables/AnimNextSharedVariables.h b/Engine/Plugins/Experimental/AnimNext/Source/AnimNext/Public/Variables/AnimNextSharedVariables.h
new file mode 100644
--- /dev/null
+++ b/Engine/Plugins/Experimental/AnimNext/Source/AnimNext/Public/Variables/AnimNextSharedVariables.h
@@ -0,0 +1,41 @@
+// Copyright Epic Games, Inc. All Rights Reserved.
+
+#pragma once
+
+#include "CoreMinimal.h"
+#include "AnimNextRigVMAsset.h"
+#include "AnimNextSharedVariables.generated.h"
+
+class UAnimNextSharedVariablesFactory;
+class UAnimNextSharedVariables;
+
+namespace UE::AnimNext
+{
+struct FInjectionInfo;
+}
+
+namespace UE::AnimNext::Tests
+{
+class FDataInterfaceCompile;
+}
+
+namespace UE::AnimNext::UncookedOnly
+{
+struct FUtils;
+}
+
+// Shared variables are just bundles of variables and their default values
+UCLASS(BlueprintType)
+class ANIMNEXT_API UAnimNextSharedVariables : public UAnimNextRigVMAsset
+{
+	GENERATED_BODY()
+
+public:
+	UAnimNextSharedVariables(const FObjectInitializer& ObjectInitializer);
+
+private:
+	friend class UAnimNextSharedVariablesFactory;
+	friend class UE::AnimNext::Tests::FDataInterfaceCompile;
+	friend struct UE::AnimNext::UncookedOnly::FUtils;
+	friend UE::AnimNext::FInjectionInfo;
+};
diff --git a/Engine/Plugins/Experimental/AnimNext/Source/AnimNext/Public/Variables/AnimNextSoftVariableReference.h b/Engine/Plugins/Experimental/AnimNext/Source/AnimNext/Public/Variables/AnimNextSoftVariableReference.h
new file mode 100644
--- /dev/null
+++ b/Engine/Plugins/Experimental/AnimNext/Source/AnimNext/Public/Variables/AnimNextSoftVariableReference.h
@@ -0,0 +1,75 @@
+﻿// Copyright Epic Games, Inc. All Rights Reserved.
+
+#pragma once
+
+#include "AnimNextSoftVariableReference.generated.h"
+
+class UAnimNextRigVMAsset;
+struct FAnimNextVariableReference;
+
+// A soft reference to an AnimNext variable
+USTRUCT(BlueprintType)
+struct FAnimNextSoftVariableReference
+{
+	GENERATED_BODY()
+
+	FAnimNextSoftVariableReference() = default;
+
+	// Construct a soft reference to an asset/struct variable
+	ANIMNEXT_API explicit FAnimNextSoftVariableReference(FName InName, const FSoftObjectPath& InSoftObjectPath)
+		: Name(InName)
+		, SoftObjectPath(InSoftObjectPath)
+	{}
+
+	// Construct a soft reference from a hard reference
+	ANIMNEXT_API explicit FAnimNextSoftVariableReference(const FAnimNextVariableReference& InVariableReference);
+
+	// Get the name of the variable
+	FName GetName() const
+	{
+		return Name;
+	}
+
+	// Get the asset or struct that the variable reference is contained in
+	const FSoftObjectPath& GetSoftObjectPath() const
+	{
+		return SoftObjectPath;
+	}
+
+	// Check whether this can ever refer to a valid variable
+	bool IsNone() const
+	{
+		return Name.IsNone() || SoftObjectPath.IsNull();
+	}
+
+	// Set this reference to None
+	void Reset()
+	{
+		Name = NAME_None; 
+		SoftObjectPath.Reset();
+	}
+
+	bool operator==(const FAnimNextSoftVariableReference& InOther) const
+	{
+		return Name == InOther.Name && SoftObjectPath == InOther.SoftObjectPath;
+	}
+
+	bool operator!=(const FAnimNextSoftVariableReference& InOther) const
+	{
+		return !(*this == InOther);
+	}
+
+	friend uint32 GetTypeHash(const FAnimNextSoftVariableReference& InKey)
+	{
+		return HashCombineFast(GetTypeHash(InKey.Name), GetTypeHash(InKey.SoftObjectPath));
+	}
+
+private:
+	// The name of the variable
+	UPROPERTY(EditAnywhere, Category = Variable)
+	FName Name;
+
+	// The asset or struct that the variable reference is contained in
+	UPROPERTY(EditAnywhere, Category = Variable)
+	FSoftObjectPath SoftObjectPath;
+};
\ No newline at end of file
diff --git a/Engine/Plugins/Experimental/AnimNext/Source/AnimNextEditor/Internal/AnimNextStyle.h b/Engine/Plugins/Experimental/AnimNext/Source/AnimNextEditor/Internal/AnimNextStyle.h
new file mode 100644
--- /dev/null
+++ b/Engine/Plugins/Experimental/AnimNext/Source/AnimNextEditor/Internal/AnimNextStyle.h
@@ -0,0 +1,20 @@
+﻿// Copyright Epic Games, Inc. All Rights Reserved.
+
+#pragma once
+
+#include "Styling/SlateStyle.h"
+
+class ISlateStyle;
+
+class ANIMNEXTEDITOR_API FAnimNextStyle : public FSlateStyleSet
+{
+public:
+	static FAnimNextStyle& Get();
+
+protected:
+	static void Register();
+	static void Unregister();
+
+private:
+	FAnimNextStyle();
+};
diff --git a/Engine/Plugins/Experimental/AnimNext/Source/AnimNextEditor/Internal/Common/AnimNextAssetItemDetails.cpp b/Engine/Plugins/Experimental/AnimNext/Source/AnimNextEditor/Internal/Common/AnimNextAssetItemDetails.cpp
--- a/Engine/Plugins/Experimental/AnimNext/Source/AnimNextEditor/Internal/Common/AnimNextAssetItemDetails.cpp
+++ b/Engine/Plugins/Experimental/AnimNext/Source/AnimNextEditor/Internal/Common/AnimNextAssetItemDetails.cpp
@@ -17,7 +17,7 @@
 #include "ToolMenus.h"
 #include "AssetRegistry/AssetRegistryModule.h"
 #include "AssetRegistry/IAssetRegistry.h"
-#include "Entries/AnimNextDataInterfaceEntry.h"
+#include "Entries/AnimNextSharedVariablesEntry.h"
 #include "Framework/Application/SlateApplication.h"
 #include "Framework/Commands/GenericCommands.h"
 #include "Module/AnimNextModule_EditorData.h"
@@ -255,38 +255,52 @@ void FAnimNextAssetItemDetails::RegisterToolMenuExtensions()
 									SNew(SAddVariablesDialog, TArray<UAnimNextRigVMAssetEditorData*>({ UncookedOnly::FUtils::GetEditorData<UAnimNextRigVMAssetEditorData>(OutlinerData.GetAsset()) }));
 
 								TArray<FVariableToAdd> VariablesToAdd;
-								TArray<FDataInterfaceToAdd> DataInterfacesToAdd;
-								if(AddVariablesDialog->ShowModal(VariablesToAdd, DataInterfacesToAdd))
+								TArray<FSharedVariablesToAdd> AllSharedVariablesToAdd;
+								if(AddVariablesDialog->ShowModal(VariablesToAdd, AllSharedVariablesToAdd))
 								{
 									FScopedTransaction Transaction(LOCTEXT("AddVariables", "Add Variable(s)"));
 									for (const FVariableToAdd& VariableToAdd : VariablesToAdd)
 									{
 										EditorData->AddVariable(VariableToAdd.Name, VariableToAdd.Type);
 									}
-									for (const FDataInterfaceToAdd& DataInterfaceToAdd : DataInterfacesToAdd)
+									for (const FSharedVariablesToAdd& SharedVariablesToAdd : AllSharedVariablesToAdd)
 									{
-										EditorData->AddDataInterface(DataInterfaceToAdd.DataInterface);
+										if (SharedVariablesToAdd.SharedVariables.IsType<UAnimNextSharedVariables*>())
+										{
+											EditorData->AddSharedVariables(SharedVariablesToAdd.SharedVariables.Get<UAnimNextSharedVariables*>());
+										}
+										else
+										{
+											EditorData->AddSharedVariablesStruct(SharedVariablesToAdd.SharedVariables.Get<const UScriptStruct*>());
+										}
 									}
 								}
 							}
-							else if(SubEntryClass == UAnimNextDataInterfaceEntry::StaticClass())
+							else if(SubEntryClass == UAnimNextSharedVariablesEntry::StaticClass())
 							{
 								TSharedRef<SAddVariablesDialog> AddVariablesDialog =
 									SNew(SAddVariablesDialog, TArray<UAnimNextRigVMAssetEditorData*>({ UncookedOnly::FUtils::GetEditorData<UAnimNextRigVMAssetEditorData>(OutlinerData.GetAsset()) }))
 									.ShouldAddInitialVariable(false);
 
 								TArray<FVariableToAdd> VariablesToAdd;
-								TArray<FDataInterfaceToAdd> DataInterfacesToAdd;
-								if(AddVariablesDialog->ShowModal(VariablesToAdd, DataInterfacesToAdd))
+								TArray<FSharedVariablesToAdd> AllSharedVariablesToAdd;
+								if(AddVariablesDialog->ShowModal(VariablesToAdd, AllSharedVariablesToAdd))
 								{
 									FScopedTransaction Transaction(LOCTEXT("AddVariables", "Add Variable(s)"));
 									for (const FVariableToAdd& VariableToAdd : VariablesToAdd)
 									{
 										EditorData->AddVariable(VariableToAdd.Name, VariableToAdd.Type);
 									}
-									for (const FDataInterfaceToAdd& DataInterfaceToAdd : DataInterfacesToAdd)
+									for (const FSharedVariablesToAdd& SharedVariablesToAdd : AllSharedVariablesToAdd)
 									{
-										EditorData->AddDataInterface(DataInterfaceToAdd.DataInterface);
+										if (SharedVariablesToAdd.SharedVariables.IsType<UAnimNextSharedVariables*>())
+										{
+											EditorData->AddSharedVariables(SharedVariablesToAdd.SharedVariables.Get<UAnimNextSharedVariables*>());
+										}
+										else
+										{
+											EditorData->AddSharedVariablesStruct(SharedVariablesToAdd.SharedVariables.Get<const UScriptStruct*>());
+										}
 									}
 								}
 							}
diff --git a/Engine/Plugins/Experimental/AnimNext/Source/AnimNextEditor/Internal/Common/GraphEditorSchemaActions.h b/Engine/Plugins/Experimental/AnimNext/Source/AnimNextEditor/Internal/Common/GraphEditorSchemaActions.h
--- a/Engine/Plugins/Experimental/AnimNext/Source/AnimNextEditor/Internal/Common/GraphEditorSchemaActions.h
+++ b/Engine/Plugins/Experimental/AnimNext/Source/AnimNextEditor/Internal/Common/GraphEditorSchemaActions.h
@@ -95,7 +95,7 @@ struct FAnimNextSchemaAction_Variable : public FAnimNextSchemaAction
 
 	FAnimNextSchemaAction_Variable() = default;
 
-	UE_API FAnimNextSchemaAction_Variable(FName InName, const FAnimNextParamType& InType, const EVariableAccessorChoice InVariableAccessorChoice);
+	UE_API FAnimNextSchemaAction_Variable(FName InName, const UObject* InSourceObject, const FAnimNextParamType& InType, const EVariableAccessorChoice InVariableAccessorChoice);
 
 	// FEdGraphSchemaAction Interface
 	UE_API virtual UEdGraphNode* PerformAction(UEdGraph* ParentGraph, TArray<UEdGraphPin*>& FromPins, const FVector2f& Location, bool bSelectNewNode = true) override;
@@ -107,7 +107,8 @@ struct FAnimNextSchemaAction_Variable : public FAnimNextSchemaAction
 
 private:
 	FName Name;
-	FString ObjectPath;
+	FString TypeObjectPath;
+	FString SourceObjectPath;
 	FString TypeName;
 	EVariableAccessorChoice VariableAccessorChoice = EVariableAccessorChoice::Set;
 	FLinearColor VariableColor;
diff --git a/Engine/Plugins/Experimental/AnimNext/Source/AnimNextEditor/Private/AnimNextEditorModule.cpp b/Engine/Plugins/Experimental/AnimNext/Source/AnimNextEditor/Private/AnimNextEditorModule.cpp
--- a/Engine/Plugins/Experimental/AnimNext/Source/AnimNextEditor/Private/AnimNextEditorModule.cpp
+++ b/Engine/Plugins/Experimental/AnimNext/Source/AnimNextEditor/Private/AnimNextEditorModule.cpp
@@ -51,7 +51,7 @@
 #include "Framework/Docking/LayoutExtender.h"
 #include "Common/AnimNextAssetItemDetails.h"
 #include "Common/AnimNextRigVMAssetEditorDataCustomization.h"
-#include "DataInterface/AnimNextDataInterface.h"
+#include "Variables/AnimNextSharedVariables.h"
 #include "Editor/RigVMEditorStyle.h"
 #include "Tools/EdModeInteractiveToolsContext.h"
 #include "Entries/AnimNextEventGraphEntry.h"
@@ -60,7 +60,6 @@
 #include "Module/RigUnit_AnimNextModuleEvents.h"
 #include "Variables/SVariablesView.h"
 #include "Variables/VariableOverrideCommands.h"
-#include "Variables/VariableProxyCustomization.h"
 #include "Workspace/AnimNextWorkspaceEditorMode.h"
 #include "Workspace/AnimNextWorkspaceSchema.h"
 #include "ExternalPackageHelper.h"
@@ -76,6 +75,7 @@
 #include "RewindDebugger/RewindDebuggerAnimNext.h"
 #include "RewindDebugger/AnimNextModuleTrack.h"
 #include "RewindDebugger/AnimNextTraceModule.h"
+#include "Variables/VariableReferencePropertyCustomization.h"
 
 #define LOCTEXT_NAMESPACE "AnimNextEditorModule"
 
@@ -125,6 +125,14 @@ void FAnimNextEditorModule::StartupModule()
 		"AnimNextVariableBinding",
 		FOnGetPropertyTypeCustomizationInstance::CreateLambda([] { return MakeShared<FVariableBindingPropertyCustomization>(); }));
 
+	PropertyModule.RegisterCustomPropertyTypeLayout(
+		"AnimNextVariableReference",
+		FOnGetPropertyTypeCustomizationInstance::CreateLambda([] { return MakeShared<FVariableReferencePropertyCustomization>(); }));
+
+	PropertyModule.RegisterCustomPropertyTypeLayout(
+		"AnimNextSoftVariableReference",
+		FOnGetPropertyTypeCustomizationInstance::CreateLambda([] { return MakeShared<FSoftVariableReferencePropertyCustomization>(); }));
+	
 	ModuleEventPropertyTypeIdentifier = MakeShared<FModuleEventPropertyTypeIdentifier>();
 	PropertyModule.RegisterCustomPropertyTypeLayout(
 		"NameProperty",
@@ -134,9 +142,6 @@ void FAnimNextEditorModule::StartupModule()
 	PropertyModule.RegisterCustomClassLayout("AnimNextVariableEntry", 
 		FOnGetDetailCustomizationInstance::CreateLambda([] { return MakeShared<FVariableCustomization>(); }));
 
-	PropertyModule.RegisterCustomClassLayout("AnimNextVariableEntryProxy", 
-		FOnGetDetailCustomizationInstance::CreateLambda([] { return MakeShared<FVariableProxyCustomization>(); }));
-
 	PropertyModule.RegisterCustomClassLayout("AnimNextRigVMAssetEditorData",
 		FOnGetDetailCustomizationInstance::CreateLambda([] { return MakeShared<FAnimNextRigVMAssetEditorDataCustomization>(); }));
 
@@ -360,7 +365,7 @@ void FAnimNextEditorModule::StartupModule()
 	SupportedAssetClasses.Append(
 		{
 			UAnimNextModule::StaticClass()->GetClassPathName(),
-			UAnimNextDataInterface::StaticClass()->GetClassPathName()
+			UAnimNextSharedVariables::StaticClass()->GetClassPathName()
 		});
 
 	FPersonaModule& PersonaModule = FModuleManager::LoadModuleChecked<FPersonaModule>("Persona");
@@ -382,9 +387,10 @@ void FAnimNextEditorModule::ShutdownModule()
 		FPropertyEditorModule& PropertyModule = FModuleManager::LoadModuleChecked<FPropertyEditorModule>("PropertyEditor");
 		PropertyModule.UnregisterCustomPropertyTypeLayout("AnimNextParamType");
 		PropertyModule.UnregisterCustomPropertyTypeLayout("AnimNextVariableBinding");
+		PropertyModule.UnregisterCustomPropertyTypeLayout("AnimNextVariableReference");
+		PropertyModule.UnregisterCustomPropertyTypeLayout("AnimNextSoftVariableReference");
 		PropertyModule.UnregisterCustomPropertyTypeLayout("NameProperty", ModuleEventPropertyTypeIdentifier);
 		PropertyModule.UnregisterCustomClassLayout("AnimNextVariableEntry");
-		PropertyModule.UnregisterCustomClassLayout("AnimNextVariableEntryProxy");
 		PropertyModule.UnregisterCustomClassLayout("AnimNextRigVMAssetEditorData");
 	}
 
@@ -628,12 +634,12 @@ void FAnimNextEditorModule::CollectGraphMenuActions(const TWeakPtr<UE::Workspace
 			}
 		}
 
-		TArray<UAnimNextVariableEntry*> Variables;
-		EditorData->GetAllVariables(Variables);
-		for (UAnimNextVariableEntry* VariableEntry : Variables)
+		TArray<UAnimNextRigVMAssetEditorData::FVariableInfo> VariableInfoArray;
+		EditorData->GetAllVariables(VariableInfoArray, UAnimNextRigVMAssetEditorData::EVariableRecursion::IncludeShared, UAnimNextRigVMAssetEditorData::EVariableAccessFilter::All);
+		for (const UAnimNextRigVMAssetEditorData::FVariableInfo& VariableInfo : VariableInfoArray)
 		{
-			InContexMenuBuilder.AddAction(MakeShared<FAnimNextSchemaAction_Variable>(VariableEntry->GetVariableName(), VariableEntry->GetType(), FAnimNextSchemaAction_Variable::EVariableAccessorChoice::Set));
-			InContexMenuBuilder.AddAction(MakeShared<FAnimNextSchemaAction_Variable>(VariableEntry->GetVariableName(), VariableEntry->GetType(), FAnimNextSchemaAction_Variable::EVariableAccessorChoice::Get));
+			InContexMenuBuilder.AddAction(MakeShared<FAnimNextSchemaAction_Variable>(VariableInfo.Name, VariableInfo.SourceObject, VariableInfo.Type, FAnimNextSchemaAction_Variable::EVariableAccessorChoice::Set));
+			InContexMenuBuilder.AddAction(MakeShared<FAnimNextSchemaAction_Variable>(VariableInfo.Name, VariableInfo.SourceObject, VariableInfo.Type, FAnimNextSchemaAction_Variable::EVariableAccessorChoice::Get));
 		}
 
 		InContexMenuBuilder.AddAction(MakeShared<FAnimNextSchemaAction_AddComment>());
