AnimNext: Pass-by-ref runtime refactor

diff --git a/Engine/Plugins/Experimental/AnimNext/Source/AnimNextEditor/Private/AnimNextStyle.cpp b/Engine/Plugins/Experimental/AnimNext/Source/AnimNextEditor/Private/AnimNextStyle.cpp
new file mode 100644
--- /dev/null
+++ b/Engine/Plugins/Experimental/AnimNext/Source/AnimNextEditor/Private/AnimNextStyle.cpp
@@ -0,0 +1,63 @@
+﻿// Copyright Epic Games, Inc. All Rights Reserved.
+
+#include "AnimNextStyle.h"
+#include "Brushes/SlateBoxBrush.h"
+#include "Styling/SlateStyleRegistry.h"
+#include "Brushes/SlateImageBrush.h"
+#include "Styling/CoreStyle.h"
+#include "Brushes/SlateRoundedBoxBrush.h"
+#include "Styling/SlateTypes.h"
+#include "Misc/Paths.h"
+#include "Styling/StyleColors.h"
+#include "Styling/SlateStyleMacros.h"
+
+
+FAnimNextStyle::FAnimNextStyle()
+	: FSlateStyleSet(TEXT("AnimNextStyle"))
+{
+	const FString EngineSlateContentDir = FPaths::EngineContentDir() / TEXT("Slate");
+	const FString EngineEditorSlateContentDir = FPaths::EngineContentDir() / TEXT("Editor/Slate");
+	SetCoreContentRoot(EngineSlateContentDir);
+
+	const FString AnimNextPluginContentDir = FPaths::EnginePluginsDir() / TEXT("Experimental/AnimNext/Resources");
+	SetContentRoot(AnimNextPluginContentDir);
+
+	const FTextBlockStyle& NormalText = FAppStyle::Get().GetWidgetStyle<FTextBlockStyle>("NormalText");
+	const FComboButtonStyle& ComboButtonStyle = FCoreStyle::Get().GetWidgetStyle<FComboButtonStyle>("ComboButton");
+
+	// Trait details labels
+	Set("AnimNext.TraitDetails.Label", FTextBlockStyle(NormalText)
+		.SetFont(FAppStyle::GetFontStyle(TEXT("PropertyWindow.BoldFont")))
+		.SetFontSize(7));
+
+	Set("AnimNext.TraitDetails.Background", new FSlateRoundedBoxBrush(FStyleColors::Hover, 6.f));
+
+	// Add trait combo button
+	const FButtonStyle AddButtonStyle = FButtonStyle()
+		.SetNormal(FSlateRoundedBoxBrush(FStyleColors::Hover, 3.0f))
+		.SetHovered(FSlateRoundedBoxBrush(FStyleColors::Primary, 3.0f))
+		.SetPressed(FSlateRoundedBoxBrush(FStyleColors::Primary, 3.0f))
+		.SetNormalForeground(FStyleColors::Transparent)
+		.SetHoveredForeground(FStyleColors::Hover)
+		.SetPressedForeground(FStyleColors::Foreground)
+		.SetNormalPadding(FMargin(6, 2, 6, 2))
+		.SetPressedPadding(FMargin(6, 3, 6, 1));
+
+	Set("AnimNext.TraitDetails.AddButton", FComboButtonStyle(ComboButtonStyle).SetButtonStyle(AddButtonStyle));
+}
+
+void FAnimNextStyle::Register()
+{
+	FSlateStyleRegistry::RegisterSlateStyle(Get());
+}
+
+void FAnimNextStyle::Unregister()
+{
+	FSlateStyleRegistry::UnRegisterSlateStyle(Get());
+}
+
+FAnimNextStyle& FAnimNextStyle::Get()
+{
+	static FAnimNextStyle Instance;
+	return Instance;
+}
diff --git a/Engine/Plugins/Experimental/AnimNext/Source/AnimNextEditor/Private/Common/AnimNextEdGraphNodeCustomization.cpp b/Engine/Plugins/Experimental/AnimNext/Source/AnimNextEditor/Private/Common/AnimNextEdGraphNodeCustomization.cpp
--- a/Engine/Plugins/Experimental/AnimNext/Source/AnimNextEditor/Private/Common/AnimNextEdGraphNodeCustomization.cpp
+++ b/Engine/Plugins/Experimental/AnimNext/Source/AnimNextEditor/Private/Common/AnimNextEdGraphNodeCustomization.cpp
@@ -461,6 +461,10 @@ void FAnimNextEdGraphNodeCustomization::PopulateCategory(IDetailLayoutBuilder& D
 										break;
 									}
 								}
+
+								// Editing some properties can cause pin-regeneration (e.g. via traits GetProgrammaticPins) so repopulate the node
+								// TODO: We could probably constrain this with metadata in case this is too disruptive
+								EdGraphNode->GetController()->RepopulatePinsOnNode(EdGraphNode->GetModelNode());
 							}
 						}
 					}
diff --git a/Engine/Plugins/Experimental/AnimNext/Source/AnimNextEditor/Private/Common/AnimNextEntryAssetDefinitions.cpp b/Engine/Plugins/Experimental/AnimNext/Source/AnimNextEditor/Private/Common/AnimNextEntryAssetDefinitions.cpp
--- a/Engine/Plugins/Experimental/AnimNext/Source/AnimNextEditor/Private/Common/AnimNextEntryAssetDefinitions.cpp
+++ b/Engine/Plugins/Experimental/AnimNext/Source/AnimNextEditor/Private/Common/AnimNextEntryAssetDefinitions.cpp
@@ -16,10 +16,10 @@ FText UAssetDefinition_AnimNextEventGraphEntry::GetObjectDisplayNameText(UObject
 	return Variable->GetDisplayName();
 }
 
-FText UAssetDefinition_AnimNextDataInterfaceEntry::GetObjectDisplayNameText(UObject* Object) const
+FText UAssetDefinition_AnimNextSharedVariablesEntry::GetObjectDisplayNameText(UObject* Object) const
 {
-	UAnimNextDataInterfaceEntry* DataInterface = CastChecked<UAnimNextDataInterfaceEntry>(Object);
-	return DataInterface->GetDisplayName();
+	UAnimNextSharedVariablesEntry* SharedVariables = CastChecked<UAnimNextSharedVariablesEntry>(Object);
+	return SharedVariables->GetDisplayName();
 }
 
 #undef LOCTEXT_NAMESPACE
\ No newline at end of file
diff --git a/Engine/Plugins/Experimental/AnimNext/Source/AnimNextEditor/Private/Common/AnimNextEntryAssetDefinitions.h b/Engine/Plugins/Experimental/AnimNext/Source/AnimNextEditor/Private/Common/AnimNextEntryAssetDefinitions.h
--- a/Engine/Plugins/Experimental/AnimNext/Source/AnimNextEditor/Private/Common/AnimNextEntryAssetDefinitions.h
+++ b/Engine/Plugins/Experimental/AnimNext/Source/AnimNextEditor/Private/Common/AnimNextEntryAssetDefinitions.h
@@ -5,7 +5,7 @@
 #include "AssetDefinitionDefault.h"
 #include "Entries/AnimNextEventGraphEntry.h"
 #include "Entries/AnimNextVariableEntry.h"
-#include "Entries/AnimNextDataInterfaceEntry.h"
+#include "Entries/AnimNextSharedVariablesEntry.h"
 #include "AnimNextEntryAssetDefinitions.generated.h"
 
 #define LOCTEXT_NAMESPACE "AnimNextAssetDefinitions"
@@ -37,15 +37,15 @@ public:
 };
 
 UCLASS()
-class UAssetDefinition_AnimNextDataInterfaceEntry : public UAssetDefinitionDefault
+class UAssetDefinition_AnimNextSharedVariablesEntry : public UAssetDefinitionDefault
 {
 	GENERATED_BODY()
 
 public:
 	// UAssetDefinition interface
-	virtual FText GetAssetDisplayName() const override { return LOCTEXT("UAFDataInterfaceEntry", "Data Interface"); }
+	virtual FText GetAssetDisplayName() const override { return LOCTEXT("UAFSharedVariablesEntry", "Shared Variables"); }
 	virtual FLinearColor GetAssetColor() const override { return FLinearColor(FColor(128,64,64)); }
-	virtual TSoftClassPtr<UObject> GetAssetClass() const override { return UAnimNextDataInterfaceEntry::StaticClass(); }
+	virtual TSoftClassPtr<UObject> GetAssetClass() const override { return UAnimNextSharedVariablesEntry::StaticClass(); }
 	virtual FText GetObjectDisplayNameText(UObject* Object) const override;
 };
 
diff --git a/Engine/Plugins/Experimental/AnimNext/Source/AnimNextEditor/Private/Common/GraphEditorSchemaActions.cpp b/Engine/Plugins/Experimental/AnimNext/Source/AnimNextEditor/Private/Common/GraphEditorSchemaActions.cpp
--- a/Engine/Plugins/Experimental/AnimNext/Source/AnimNextEditor/Private/Common/GraphEditorSchemaActions.cpp
+++ b/Engine/Plugins/Experimental/AnimNext/Source/AnimNextEditor/Private/Common/GraphEditorSchemaActions.cpp
@@ -1,6 +1,8 @@
 ﻿// Copyright Epic Games, Inc. All Rights Reserved.
 
 #include "Common/GraphEditorSchemaActions.h"
+
+#include "AnimNextControllerBase.h"
 #include "EditorUtils.h"
 #include "AnimNextEdGraph.h"
 #include "AnimNextEdGraphNode.h"
@@ -13,8 +15,9 @@
 #include "Framework/MultiBox/MultiBoxBuilder.h"
 #include "Framework/Application/SlateApplication.h"
 #include "AnimNextUnitNode.h"
+#include "AnimNextSharedVariableNode.h"
 #include "Editor/RigVMEditorStyle.h"
-
+#include "Variables/AnimNextSharedVariables.h"
 
 #define LOCTEXT_NAMESPACE "AnimNextSchemaActions"
 
@@ -145,13 +148,16 @@ UEdGraphNode* FAnimNextSchemaAction_DispatchFactory::PerformAction(UEdGraph* Par
 
 // *** Variable ***
 
-FAnimNextSchemaAction_Variable::FAnimNextSchemaAction_Variable(FName InName, const FAnimNextParamType& InType, const EVariableAccessorChoice InVariableAccessorChoice)
+FAnimNextSchemaAction_Variable::FAnimNextSchemaAction_Variable(FName InName, const UObject* InSourceObject, const FAnimNextParamType& InType, const EVariableAccessorChoice InVariableAccessorChoice)
 	: Name(InName)
 	, VariableAccessorChoice(InVariableAccessorChoice)
 {
+	check(InSourceObject->IsA<UAnimNextSharedVariables>() || InSourceObject->IsA<UScriptStruct>());
+	SourceObjectPath = InSourceObject->GetPathName();
+
 	if(InType.IsObjectType())
 	{
-		ObjectPath = InType.GetValueTypeObject()->GetPathName();
+		TypeObjectPath = InType.GetValueTypeObject()->GetPathName();
 	}
 
 	TypeName = InType.ToRigVMTemplateArgument().GetBaseCPPType();
@@ -196,7 +202,7 @@ const FLinearColor& FAnimNextSchemaAction_Variable::GetIconColor() const
 
 UEdGraphNode* FAnimNextSchemaAction_Variable::PerformAction(UEdGraph* ParentGraph, TArray<UEdGraphPin*>& FromPins, const FVector2f& Location, bool bSelectNewNode)
 {
-	const auto AddNodeLambda = [](UEdGraph* ParentGraph, FName InName, const FString& InTypeName, const FString& InObjectPath, const FVector2f& Location, const bool bIsGetter) -> URigVMEdGraphNode*
+	const auto AddNodeLambda = [](UEdGraph* ParentGraph, FName InName, const FString& InSourceObjectPath, const FString& InTypeName, const FString& InTypeObjectPath, const FVector2f& Location, const bool bIsGetter) -> URigVMEdGraphNode*
 	{
 		IRigVMClientHost* Host = ParentGraph->GetImplementingOuter<IRigVMClientHost>();
 		if(Host == nullptr)
@@ -214,10 +220,42 @@ UEdGraphNode* FAnimNextSchemaAction_Variable::PerformAction(UEdGraph* ParentGrap
 
 		FString NodeName;
 
-		URigVMController* Controller = Host->GetRigVMClient()->GetController(ParentGraph);
+		UAnimNextControllerBase* Controller = CastChecked<UAnimNextControllerBase>(Host->GetRigVMClient()->GetController(ParentGraph));
 		Controller->OpenUndoBracket(TEXT("Add Variable"));
 
-		if (URigVMNode* ModelNode = Controller->AddVariableNodeFromObjectPath(InName, InTypeName, InObjectPath, bIsGetter, FString(), FDeprecateSlateVector2D(Location), NodeName, true, true))
+		UObject* CPPTypeObject = nullptr;
+		if (!InTypeObjectPath.IsEmpty())
+		{
+			CPPTypeObject = RigVMTypeUtils::FindObjectFromCPPTypeObjectPath<UObject>(InTypeObjectPath);
+			if (CPPTypeObject == nullptr)
+			{
+				return nullptr;
+			}
+		}
+
+		UObject* SourceObject = nullptr;
+		if (!InSourceObjectPath.IsEmpty())
+		{
+			SourceObject = RigVMTypeUtils::FindObjectFromCPPTypeObjectPath<UObject>(InSourceObjectPath);
+			if (SourceObject == nullptr)
+			{
+				return nullptr;
+			}
+		}
+
+		URigVMNode* ModelNode = nullptr;
+
+		// If the variable is from 'this' asset, we just need a standard variable node
+		if (ParentGraph->GetTypedOuter<UAnimNextSharedVariables>() == SourceObject)
+		{
+			ModelNode = Controller->AddVariableNode(InName, InTypeName, CPPTypeObject, bIsGetter, FString(), FDeprecateSlateVector2D(Location), NodeName, true, true);
+		}
+		else
+		{
+			ModelNode = Controller->AddSharedVariableNode(InSourceObjectPath, InName, InTypeName, CPPTypeObject, bIsGetter, FString(), FDeprecateSlateVector2D(Location), NodeName, true, true);
+		}
+
+		if (ModelNode)
 		{
 			for (UEdGraphNode* Node : ParentGraph->Nodes)
 			{
@@ -258,9 +296,9 @@ UEdGraphNode* FAnimNextSchemaAction_Variable::PerformAction(UEdGraph* ParentGrap
 			FText::FromString(FString::Printf(TEXT("Adds a getter node for variable %s"), *Name.ToString())),
 			FSlateIcon(),
 			FUIAction(
-				FExecuteAction::CreateLambda([AddNodeLambda, ParentGraph, Location, Name = Name, TypeName = TypeName, ObjectPath = ObjectPath]
+				FExecuteAction::CreateLambda([AddNodeLambda, ParentGraph, Location, Name = Name, SourceObjectPath = SourceObjectPath, TypeName = TypeName, TypeObjectPath = TypeObjectPath]
 				{
-					AddNodeLambda(ParentGraph, Name, TypeName, ObjectPath, Location, true);
+					AddNodeLambda(ParentGraph, Name, SourceObjectPath, TypeName, TypeObjectPath, Location, true);
 				}),
 				FCanExecuteAction()
 			)
@@ -271,9 +309,9 @@ UEdGraphNode* FAnimNextSchemaAction_Variable::PerformAction(UEdGraph* ParentGrap
 			FText::FromString(FString::Printf(TEXT("Adds a setter node for variable %s"), *Name.ToString())),
 			FSlateIcon(),
 			FUIAction(
-				FExecuteAction::CreateLambda([AddNodeLambda, ParentGraph, Location, Name = Name, TypeName = TypeName, ObjectPath = ObjectPath]
+				FExecuteAction::CreateLambda([AddNodeLambda, ParentGraph, Location, Name = Name, SourceObjectPath = SourceObjectPath, TypeName = TypeName, ObjectPath = TypeObjectPath]
 				{
-					AddNodeLambda(ParentGraph, Name, TypeName, ObjectPath, Location, false);
+					AddNodeLambda(ParentGraph, Name, SourceObjectPath, TypeName, ObjectPath, Location, false);
 				}),
 				FCanExecuteAction()
 			)
@@ -294,7 +332,7 @@ UEdGraphNode* FAnimNextSchemaAction_Variable::PerformAction(UEdGraph* ParentGrap
 	else
 	{
 		const bool bIsGetter = VariableAccessorChoice == EVariableAccessorChoice::Get;
-		return AddNodeLambda(ParentGraph, Name, TypeName, ObjectPath, Location, bIsGetter);
+		return AddNodeLambda(ParentGraph, Name, SourceObjectPath, TypeName, TypeObjectPath, Location, bIsGetter);
 	}
 }
 
diff --git a/Engine/Plugins/Experimental/AnimNext/Source/AnimNextEditor/Private/Graph/AnimNextGraphPanelPinFactory.cpp b/Engine/Plugins/Experimental/AnimNext/Source/AnimNextEditor/Private/Graph/AnimNextGraphPanelPinFactory.cpp
--- a/Engine/Plugins/Experimental/AnimNext/Source/AnimNextEditor/Private/Graph/AnimNextGraphPanelPinFactory.cpp
+++ b/Engine/Plugins/Experimental/AnimNext/Source/AnimNextEditor/Private/Graph/AnimNextGraphPanelPinFactory.cpp
@@ -5,6 +5,9 @@
 #include "NodeFactory.h"
 #include "SGraphPinModuleEvent.h"
 #include "EdGraph/RigVMEdGraphNode.h"
+#include "Variables/SGraphPinVariableReference.h"
+#include "Variables/AnimNextVariableReference.h"
+#include "Variables/AnimNextSoftVariableReference.h"
 #include "Widgets/DeclarativeSyntaxSupport.h"
 
 TSharedPtr<SGraphPin> FAnimNextGraphPanelPinFactory::CreatePin(UEdGraphPin* InPin) const
@@ -23,5 +26,14 @@ TSharedPtr<SGraphPin> FAnimNextGraphPanelPinFactory::CreatePin(UEdGraphPin* InPi
 		}
 	}
 
+	if (InPin->PinType.PinCategory == UEdGraphSchema_K2::PC_Struct)
+	{
+		const UScriptStruct* Struct = Cast<UScriptStruct>(InPin->PinType.PinSubCategoryObject.Get());
+		if (Struct && (Struct->IsChildOf<FAnimNextSoftVariableReference>() || Struct->IsChildOf<FAnimNextVariableReference>()))
+		{
+			return SNew(UE::AnimNext::Editor::SGraphPinVariableReference, InPin);
+		}
+	}
+
 	return nullptr;
 }
diff --git a/Engine/Plugins/Experimental/AnimNext/Source/AnimNextEditor/Private/Variables/AnimNextSharedVariablesAssetDefinition.cpp b/Engine/Plugins/Experimental/AnimNext/Source/AnimNextEditor/Private/Variables/AnimNextSharedVariablesAssetDefinition.cpp
new file mode 100644
--- /dev/null
+++ b/Engine/Plugins/Experimental/AnimNext/Source/AnimNextEditor/Private/Variables/AnimNextSharedVariablesAssetDefinition.cpp
@@ -0,0 +1,153 @@
+﻿// Copyright Epic Games, Inc. All Rights Reserved.
+
+#include "AnimNextSharedVariablesAssetDefinition.h"
+
+#include "AnimNextRigVMAssetEditorData.h"
+#include "ContentBrowserMenuContexts.h"
+#include "Editor.h"
+#include "IWorkspaceEditorModule.h"
+#include "ToolMenus.h"
+#include "UncookedOnlyUtils.h"
+#include "UObject/SavePackage.h"
+#include "Workspace/AnimNextWorkspaceFactory.h"
+
+#define LOCTEXT_NAMESPACE "AnimNextAssetDefinitions"
+
+EAssetCommandResult UAssetDefinition_AnimNextSharedVariables::OpenAssets(const FAssetOpenArgs& OpenArgs) const
+{
+	using namespace UE::Workspace;
+
+	for (UAnimNextSharedVariables* Asset : OpenArgs.LoadObjects<UAnimNextSharedVariables>())
+	{
+		IWorkspaceEditorModule& WorkspaceEditorModule = FModuleManager::Get().LoadModuleChecked<IWorkspaceEditorModule>("WorkspaceEditor");
+		WorkspaceEditorModule.OpenWorkspaceForObject(Asset, EOpenWorkspaceMethod::Default, UAnimNextWorkspaceFactory::StaticClass());
+	}
+
+	return EAssetCommandResult::Handled;
+}
+
+namespace UE::AnimNext::Editor
+{
+
+static FDelayedAutoRegisterHelper AutoRegisterGraphMenuItems(EDelayedRegisterRunPhase::EndOfEngineInit, []
+{
+	UToolMenus::RegisterStartupCallback(FSimpleMulticastDelegate::FDelegate::CreateLambda([]()
+	{
+		FToolMenuOwnerScoped OwnerScoped(UE_MODULE_NAME);
+
+		UToolMenu* Menu = UE::ContentBrowser::ExtendToolMenu_AssetContextMenu(UAnimNextSharedVariables::StaticClass());
+
+		FToolMenuSection& Section = Menu->FindOrAddSection("GetAssetActions");
+		Section.AddDynamicEntry(NAME_None, FNewToolMenuSectionDelegate::CreateLambda([](FToolMenuSection& InSection)
+		{
+			auto SetExternalPackagingStatus = [](const FToolMenuContext& InContext, bool bInUseExternalPackaging)
+			{
+				UContentBrowserAssetContextMenuContext* AssetContext = InContext.FindContext<UContentBrowserAssetContextMenuContext>();
+				if(AssetContext == nullptr)
+				{
+					return;
+				}
+
+				TArray<UAnimNextRigVMAsset*> Assets;
+				for (const FAssetData& AssetData : AssetContext->SelectedAssets)
+				{
+					if(!AssetData.GetClass()->IsChildOf(UAnimNextRigVMAsset::StaticClass()))
+					{
+						continue;
+					}
+
+					bool bUsesExternalPackages = true;
+					AssetData.GetTagValue<bool>(UAnimNextRigVMAssetEditorData::GetUsesExternalPackagesPropertyName(), bUsesExternalPackages);
+
+					if(bUsesExternalPackages == bInUseExternalPackaging)
+					{
+						// Status is as requested, skip
+						continue;
+					}
+					
+					UAnimNextRigVMAsset* Asset = Cast<UAnimNextRigVMAsset>(AssetData.GetAsset());
+					if(Asset == nullptr)
+					{
+						continue;
+					}
+
+					Assets.Add(Asset);
+
+					UAnimNextRigVMAssetEditorData* EditorData = UncookedOnly::FUtils::GetEditorData<UAnimNextRigVMAssetEditorData>(Asset);
+					if(EditorData == nullptr)
+					{
+						continue;
+					}
+
+					ensure(EditorData->IsUsingExternalPackages() == bUsesExternalPackages);
+				}
+
+				UAnimNextRigVMAssetEditorData::SetUseExternalPackages(Assets, bInUseExternalPackaging);
+			};
+
+			auto ExternalPackagingStatusMatches = [](const FToolMenuContext& InContext, bool bInUsesExternalPackaging)
+			{
+				UContentBrowserAssetContextMenuContext* AssetContext = InContext.FindContext<UContentBrowserAssetContextMenuContext>();
+				if(AssetContext == nullptr)
+				{
+					return false;
+				}
+
+				for (const FAssetData& AssetData : AssetContext->SelectedAssets)
+				{
+					if(!AssetData.GetClass()->IsChildOf(UAnimNextRigVMAsset::StaticClass()))
+					{
+						continue;
+					}
+
+					bool bUsesExternalPackages = true;
+					AssetData.GetTagValue<bool>(UAnimNextRigVMAssetEditorData::GetUsesExternalPackagesPropertyName(), bUsesExternalPackages);
+
+					if(bUsesExternalPackages == bInUsesExternalPackaging)
+					{
+						return true;
+					}
+				}
+
+				return false;
+			};
+
+			FToolUIAction UIAction;
+			UIAction.ExecuteAction = FToolMenuExecuteAction::CreateLambda([&SetExternalPackagingStatus](const FToolMenuContext& InContext)
+			{
+				SetExternalPackagingStatus(InContext, true);
+			});
+			UIAction.IsActionVisibleDelegate = FToolMenuIsActionButtonVisible::CreateLambda([&ExternalPackagingStatusMatches](const FToolMenuContext& InContext)
+			{
+				return ExternalPackagingStatusMatches(InContext, false);
+			});
+
+			InSection.AddMenuEntry(
+				"EnableExternalPackages",
+				LOCTEXT("EnableExternalPackagesLabel", "Use External Packages"),
+				LOCTEXT("EnableExternalPackagesTooltip", "Set the asset(s) to use external packaging for its entries (graphs, variables etc.)\nThis will create the external packages for all entries, add them to version control if enabled and save all packages.\nWarning: This operation cannot be undone, so a connection to version control is recommended."),
+				FSlateIcon(FAppStyle::Get().GetStyleSetName(), "Icons.Edit"),
+				UIAction);
+
+			UIAction.ExecuteAction = FToolMenuExecuteAction::CreateLambda([&SetExternalPackagingStatus](const FToolMenuContext& InContext)
+			{
+				SetExternalPackagingStatus(InContext, false);
+			});
+			UIAction.IsActionVisibleDelegate = FToolMenuIsActionButtonVisible::CreateLambda([&ExternalPackagingStatusMatches](const FToolMenuContext& InContext)
+			{
+				return ExternalPackagingStatusMatches(InContext, true);
+			});
+
+			InSection.AddMenuEntry(
+				"DisableExternalPackages",
+				LOCTEXT("DisableExternalPackagesLabel", "Use Single Package"),
+				LOCTEXT("DisableExternalPackagesTooltip", "Set the asset(s) to use a single package.\nThis will remove any external packages for existing entries, remove them from version control if enabled and save all packages.\nWarning: This operation cannot be undone, so a connection to version control is recommended."),
+				FSlateIcon(FAppStyle::Get().GetStyleSetName(), "Icons.Edit"),
+				UIAction);
+		}));
+	}));
+});
+
+}
+
+#undef LOCTEXT_NAMESPACE
\ No newline at end of file
diff --git a/Engine/Plugins/Experimental/AnimNext/Source/AnimNextEditor/Private/Variables/AnimNextSharedVariablesFactory.cpp b/Engine/Plugins/Experimental/AnimNext/Source/AnimNextEditor/Private/Variables/AnimNextSharedVariablesFactory.cpp
new file mode 100644
--- /dev/null
+++ b/Engine/Plugins/Experimental/AnimNext/Source/AnimNextEditor/Private/Variables/AnimNextSharedVariablesFactory.cpp
@@ -0,0 +1,41 @@
+﻿// Copyright Epic Games, Inc. All Rights Reserved.
+
+#include "Variables/AnimNextSharedVariablesFactory.h"
+#include "Variables/AnimNextSharedVariables.h"
+#include "Variables/AnimNextSharedVariables_EditorData.h"
+#include "Module/RigUnit_AnimNextModuleEvents.h"
+
+UAnimNextSharedVariablesFactory::UAnimNextSharedVariablesFactory()
+{
+	bCreateNew = true;
+	bEditAfterNew = true;
+	SupportedClass = UAnimNextSharedVariables::StaticClass();
+}
+
+bool UAnimNextSharedVariablesFactory::ConfigureProperties()
+{
+	return true;
+}
+
+UObject* UAnimNextSharedVariablesFactory::FactoryCreateNew(UClass* Class, UObject* InParent, FName Name, EObjectFlags Flags, UObject* Context, FFeedbackContext* Warn, FName CallingContext)
+{
+	EObjectFlags FlagsToUse = Flags | RF_Public | RF_Standalone | RF_Transactional | RF_LoadCompleted;
+	if(InParent == GetTransientPackage())
+	{
+		FlagsToUse &= ~RF_Standalone;
+	}
+
+	UAnimNextSharedVariables* NewDataInterface = NewObject<UAnimNextSharedVariables>(InParent, Class, Name, FlagsToUse);
+
+	// Create internal editor data
+	UAnimNextSharedVariables_EditorData* EditorData = NewObject<UAnimNextSharedVariables_EditorData>(NewDataInterface, TEXT("EditorData"), RF_Transactional);
+	NewDataInterface->EditorData = EditorData;
+	EditorData->bUsesExternalPackages = false;
+	EditorData->Initialize(/*bRecompileVM*/false);
+
+	// Compile the initial skeleton
+	EditorData->RecompileVM();
+	check(!EditorData->bErrorsDuringCompilation);
+
+	return NewDataInterface;
+}
\ No newline at end of file
diff --git a/Engine/Plugins/Experimental/AnimNext/Source/AnimNextEditor/Private/Variables/Outliner/VariablesOutlinerColumns.cpp b/Engine/Plugins/Experimental/AnimNext/Source/AnimNextEditor/Private/Variables/Outliner/VariablesOutlinerColumns.cpp
--- a/Engine/Plugins/Experimental/AnimNext/Source/AnimNextEditor/Private/Variables/Outliner/VariablesOutlinerColumns.cpp
+++ b/Engine/Plugins/Experimental/AnimNext/Source/AnimNextEditor/Private/Variables/Outliner/VariablesOutlinerColumns.cpp
@@ -15,14 +15,13 @@
 #include "SPinTypeSelector.h"
 #include "UncookedOnlyUtils.h"
 #include "VariablesOutlinerEntryItem.h"
-#include "Entries/AnimNextDataInterfaceEntry.h"
+#include "Entries/AnimNextSharedVariablesEntry.h"
 #include "Entries/AnimNextRigVMAssetEntry.h"
 #include "Entries/AnimNextVariableEntry.h"
 #include "Widgets/Images/SImage.h"
 #include "UObject/Package.h"
 #include "Widgets/Input/SButton.h"
 #include "Widgets/Text/SInlineEditableTextBlock.h"
-#include "Variables/SVariableOverride.h"
 
 #define LOCTEXT_NAMESPACE "VariablesOutlinerColumns"
 
@@ -71,15 +70,29 @@ const TSharedRef<SWidget> FVariablesOutlinerTypeColumn::ConstructRowWidget(FScen
 		SNew(SBox)
 		.HAlign(HAlign_Left)
 		.VAlign(VAlign_Center)
-		.IsEnabled(!TreeItem->WeakDataInterfaceEntry.IsValid())
+		.IsEnabled(!TreeItem->WeakDataInterfaceEntry.IsValid() || !TreeItem->PropertyPath.IsPathToFieldEmpty())
 		[
 			SNew(SPinTypeSelector, FGetPinTypeTree::CreateStatic(&Editor::FUtils::GetFilteredVariableTypeTree))
-			.TargetPinType_Lambda([WeakEntry = TreeItem->WeakEntry]()
+			.TargetPinType_Lambda([WeakItem = TWeakPtr<ISceneOutlinerTreeItem>(Item)]()
 			{
-				if(const IAnimNextRigVMVariableInterface* Variable = Cast<IAnimNextRigVMVariableInterface>(WeakEntry.Get()))
+				TSharedPtr<ISceneOutlinerTreeItem> PinnedItem = WeakItem.Pin();
+				if (!PinnedItem->IsValid())
+				{
+					return FEdGraphPinType();
+				}
+
+				const FVariablesOutlinerEntryItem* VariablesItem = PinnedItem->CastTo<FVariablesOutlinerEntryItem>();
+				
+				if(const IAnimNextRigVMVariableInterface* Variable = Cast<IAnimNextRigVMVariableInterface>(VariablesItem->WeakEntry.Get()))
 				{
 					return UncookedOnly::FUtils::GetPinTypeFromParamType(Variable->GetType());
 				}
+				else if (const FProperty* Property = VariablesItem->PropertyPath.Get())
+				{
+					FEdGraphPinType PinType;
+					GetDefault<UEdGraphSchema_K2>()->ConvertPropertyToPinType(Property, PinType);
+					return PinType;
+				}
 
 				return FEdGraphPinType();
 			})
@@ -147,7 +160,7 @@ class SVariablesOutlinerValue : public SCompoundWidget, public FNotifyHook
 				EditorData->ModifiedDelegate.AddSP(this, &SVariablesOutlinerValue::HandleModified);
 			}
 		}
-		if(UAnimNextDataInterfaceEntry* DataInterfaceEntry = Cast<UAnimNextDataInterfaceEntry>(InTreeItem.WeakDataInterfaceEntry.Get()))
+		if(UAnimNextSharedVariablesEntry* DataInterfaceEntry = Cast<UAnimNextSharedVariablesEntry>(InTreeItem.WeakDataInterfaceEntry.Get()))
 		{
 			if(UAnimNextRigVMAssetEditorData* EditorData = DataInterfaceEntry->GetTypedOuter<UAnimNextRigVMAssetEditorData>())
 			{
@@ -178,57 +191,10 @@ class SVariablesOutlinerValue : public SCompoundWidget, public FNotifyHook
 		}
 
 		TSharedPtr<SWidget> ValueWidget = SNullWidget::NullWidget;
-		const FName VariableName = VariableEntry->GetEntryName();
-		FInstancedPropertyBag* PropertyBag = nullptr;
-		FName ValueName;
-		bool bCanOverride = false;
-		UAnimNextDataInterfaceEntry* DataInterfaceEntry = TreeItem->WeakDataInterfaceEntry.Get();
-		if(DataInterfaceEntry != nullptr)
-		{
-			// Update for override status
-			bCanOverride = true;
-			OverrideStatus = DataInterfaceEntry->FindValueOverridePropertyBagRecursive(VariableName, PropertyBag);
-			switch(OverrideStatus)
-			{
-			case EAnimNextDataInterfaceValueOverrideStatus::NotOverridden:
-				{
-					// Use the internal property bag (copying from the source) and the default name
-					InternalPropertyBag = VariableEntry->GetPropertyBag();
-					PropertyBag = &InternalPropertyBag;
-					ValueName = IAnimNextRigVMVariableInterface::ValueName;
-					break;
-				}
-			case EAnimNextDataInterfaceValueOverrideStatus::OverriddenInThisAsset:
-				{
-					// Use the found property bag and the variable name
-					check(PropertyBag);
-					ValueName = VariableName;
-					break;
-				}
-			case EAnimNextDataInterfaceValueOverrideStatus::OverriddenInParentAsset:
-				{
-					// Use the internal property bag (copying just the value from the overriding asset's property bag) and the default name
-					check(PropertyBag);
-					const FPropertyBagPropertyDesc* Desc = PropertyBag->FindPropertyDescByName(VariableName);
-					check(Desc);
-					const uint8* DataPtr = Desc->CachedProperty->ContainerPtrToValuePtr<uint8>(PropertyBag->GetValue().GetMemory());
-					TConstArrayView<uint8> Value(DataPtr, Desc->CachedProperty->GetElementSize());
-					InternalPropertyBag.ReplaceAllPropertiesAndValues(TConstArrayView<FPropertyBagPropertyDesc>(Desc, 1), TConstArrayView<TConstArrayView<uint8>>(&Value, 1));
-					PropertyBag = &InternalPropertyBag;
-					ValueName = VariableName;
-					break;
-				}
-			}
-		}
-		else
-		{
-			// No data interface present, use the variable's internal property bag
-			bCanOverride = false;
-			PropertyBag = &VariableEntry->GetMutablePropertyBag();
-			ValueName = IAnimNextRigVMVariableInterface::ValueName;
-		}
+		FInstancedPropertyBag& PropertyBag = VariableEntry->GetMutablePropertyBag();
+		const FName ValueName = IAnimNextRigVMVariableInterface::ValueName;
 
-		const FPropertyBagPropertyDesc* PropertyDesc = PropertyBag ? PropertyBag->FindPropertyDescByName(ValueName) : nullptr;
+		const FPropertyBagPropertyDesc* PropertyDesc = PropertyBag.FindPropertyDescByName(ValueName);
 		if (PropertyDesc != nullptr)
 		{
 			if (PropertyDesc->ContainerTypes.IsEmpty()) // avoid trying to inline containers
@@ -240,7 +206,7 @@ class SVariablesOutlinerValue : public SCompoundWidget, public FNotifyHook
 
 				FPropertyEditorModule& PropertyEditorModule = FModuleManager::Get().LoadModuleChecked<FPropertyEditorModule>("PropertyEditor");
 				
-				const TSharedPtr<ISinglePropertyView> SingleStructPropertyView = PropertyEditorModule.CreateSingleProperty(MakeShared<FInstancePropertyBagStructureDataProvider>(*PropertyBag), ValueName, SinglePropertyArgs);
+				const TSharedPtr<ISinglePropertyView> SingleStructPropertyView = PropertyEditorModule.CreateSingleProperty(MakeShared<FInstancePropertyBagStructureDataProvider>(PropertyBag), ValueName, SinglePropertyArgs);
 				if (SingleStructPropertyView.IsValid())
 				{
 					ValueWidget = SNew(SHorizontalBox)
@@ -252,15 +218,6 @@ class SVariablesOutlinerValue : public SCompoundWidget, public FNotifyHook
 							[
 								SingleStructPropertyView.ToSharedRef()
 							]
-						]
-						+SHorizontalBox::Slot()
-						.AutoWidth()
-						.VAlign(VAlign_Center)
-						.HAlign(HAlign_Center)
-						[
-							SNew(SVariableOverride, DataInterfaceEntry, VariableName)
-							.Visibility(bCanOverride ? EVisibility::Visible : EVisibility::Collapsed)
-							.OverrideStatus(OverrideStatus)
 						];
 				}
 			}
@@ -290,28 +247,8 @@ class SVariablesOutlinerValue : public SCompoundWidget, public FNotifyHook
 			return;
 		}
 
-		UAnimNextDataInterfaceEntry* DataInterfaceEntry = Cast<UAnimNextDataInterfaceEntry>(TreeItem->WeakDataInterfaceEntry.Get());
-		if(DataInterfaceEntry != nullptr)
-		{
-			const FName VariableName = VariableEntry->GetVariableName();
-			if(!TreeItem->WeakDataInterfaceEntry->HasValueOverride(VariableName))
-			{
-				// No value override yet, so we copy from the internal property bag
-				check(InternalPropertyBag.IsValid() && InternalPropertyBag.GetPropertyBagStruct()->GetPropertyDescs().Num() == 1);
-				const FProperty* Property = InternalPropertyBag.GetPropertyBagStruct()->GetPropertyDescs()[0].CachedProperty; 
-				const uint8* DataPtr = Property->ContainerPtrToValuePtr<uint8>(InternalPropertyBag.GetValue().GetMemory());
-				TreeItem->WeakDataInterfaceEntry->SetValueOverride(VariableName, VariableEntry->GetType(), TConstArrayView<uint8>(DataPtr, Property->GetElementSize()));
-				ensure(TreeItem->WeakDataInterfaceEntry->HasValueOverride(VariableName));
-			}
-
-			DataInterfaceEntry->MarkPackageDirty();
-			DataInterfaceEntry->BroadcastModified(EAnimNextEditorDataNotifType::VariableDefaultValueChanged);
-		}
-		else
-		{
-			VariableEntry->MarkPackageDirty();
-			VariableEntry->BroadcastModified(EAnimNextEditorDataNotifType::VariableDefaultValueChanged);
-		}
+		VariableEntry->MarkPackageDirty();
+		VariableEntry->BroadcastModified(EAnimNextEditorDataNotifType::VariableDefaultValueChanged);
 	}
 
 	void HandleModified(UAnimNextRigVMAssetEditorData* InEditorData, EAnimNextEditorDataNotifType InType, UObject* InSubject)
@@ -329,7 +266,7 @@ class SVariablesOutlinerValue : public SCompoundWidget, public FNotifyHook
 		}
 
 		UAnimNextVariableEntry* VariableEntry = Cast<UAnimNextVariableEntry>(TreeItem->WeakEntry.Get());
-		UAnimNextDataInterfaceEntry* DataInterfaceEntry = Cast<UAnimNextDataInterfaceEntry>(TreeItem->WeakDataInterfaceEntry.Get());
+		UAnimNextSharedVariablesEntry* DataInterfaceEntry = Cast<UAnimNextSharedVariablesEntry>(TreeItem->WeakDataInterfaceEntry.Get());
 		if(VariableEntry != InSubject && DataInterfaceEntry != InSubject)
 		{
 			return;
@@ -341,7 +278,6 @@ class SVariablesOutlinerValue : public SCompoundWidget, public FNotifyHook
 	TWeakPtr<FVariablesOutlinerEntryItem> WeakTreeItem;
 	TSharedPtr<SBox> WidgetContainer;
 	FInstancedPropertyBag InternalPropertyBag;
-	EAnimNextDataInterfaceValueOverrideStatus OverrideStatus = EAnimNextDataInterfaceValueOverrideStatus::NotOverridden;
 };
 
 const TSharedRef<SWidget> FVariablesOutlinerValueColumn::ConstructRowWidget(FSceneOutlinerTreeItemRef Item, const STableRow<FSceneOutlinerTreeItemPtr>& Row)
diff --git a/Engine/Plugins/Experimental/AnimNext/Source/AnimNextEditor/Private/Variables/Outliner/VariablesOutlinerDataInterfaceItem.cpp b/Engine/Plugins/Experimental/AnimNext/Source/AnimNextEditor/Private/Variables/Outliner/VariablesOutlinerDataInterfaceItem.cpp
--- a/Engine/Plugins/Experimental/AnimNext/Source/AnimNextEditor/Private/Variables/Outliner/VariablesOutlinerDataInterfaceItem.cpp
+++ b/Engine/Plugins/Experimental/AnimNext/Source/AnimNextEditor/Private/Variables/Outliner/VariablesOutlinerDataInterfaceItem.cpp
@@ -9,8 +9,8 @@
 #include "SceneOutlinerStandaloneTypes.h"
 #include "AnimNextEditorModule.h"
 #include "AnimNextRigVMAssetEditorData.h"
-#include "DataInterface/AnimNextDataInterface.h"
-#include "Entries/AnimNextDataInterfaceEntry.h"
+#include "Variables/AnimNextSharedVariables.h"
+#include "Entries/AnimNextSharedVariablesEntry.h"
 #include "Styling/SlateIconFinder.h"
 #include "Styling/StarshipCoreStyle.h"
 #include "UObject/Package.h"
@@ -40,9 +40,9 @@ class SVariablesOutlinerDataInterfaceLabel : FSceneOutlinerCommonLabelData, publ
 		TreeItem = StaticCastSharedRef<FVariablesOutlinerDataInterfaceItem>(InTreeItem.AsShared());
 
 		FText AssetName;
-		if(UAnimNextDataInterfaceEntry* AssetEntry = InTreeItem.WeakEntry.Get())
+		if(UAnimNextSharedVariablesEntry* AssetEntry = InTreeItem.WeakEntry.Get())
 		{
-			AssetName = FText::FromString(AssetEntry->GetDataInterfacePath().ToString());
+			AssetName = FText::FromString(AssetEntry->GetObjectPath().ToString());
 		}
 		else
 		{
@@ -52,7 +52,7 @@ class SVariablesOutlinerDataInterfaceLabel : FSceneOutlinerCommonLabelData, publ
 		ChildSlot
 		[
 			SNew(SHorizontalBox)
-			.ToolTipText(FText::Format(LOCTEXT("ImportedVariablesFormat", "Imported variables from '{0}'"), AssetName))
+			.ToolTipText(FText::Format(LOCTEXT("ImportedVariablesFormat", "Shared variables from '{0}'"), AssetName))
 			+SHorizontalBox::Slot()
 			.AutoWidth()
 			.VAlign(VAlign_Center)
@@ -92,7 +92,7 @@ class SVariablesOutlinerDataInterfaceLabel : FSceneOutlinerCommonLabelData, publ
 			FTextBuilder TextBuilder;
 			TextBuilder.AppendLine(LOCTEXT("ModifiedTooltip", "Modified"));
 
-			if(UAnimNextDataInterfaceEntry* AssetEntry = Item->WeakEntry.Get())
+			if(UAnimNextSharedVariablesEntry* AssetEntry = Item->WeakEntry.Get())
 			{
 				const UPackage* ExternalPackage = AssetEntry->GetExternalPackage();
 				if(ExternalPackage && ExternalPackage->IsDirty())
@@ -110,7 +110,7 @@ class SVariablesOutlinerDataInterfaceLabel : FSceneOutlinerCommonLabelData, publ
 	{
 		if (const TSharedPtr<FVariablesOutlinerDataInterfaceItem> Item = TreeItem.Pin())
 		{
-			if(UAnimNextDataInterfaceEntry* AssetEntry = Item->WeakEntry.Get())
+			if(UAnimNextSharedVariablesEntry* AssetEntry = Item->WeakEntry.Get())
 			{
 				bool bIsDirty = false;
 				const UPackage* ExternalPackage = AssetEntry->GetExternalPackage();
@@ -129,7 +129,7 @@ class SVariablesOutlinerDataInterfaceLabel : FSceneOutlinerCommonLabelData, publ
 	{
 		if (const TSharedPtr<FVariablesOutlinerDataInterfaceItem> Item = TreeItem.Pin())
 		{
-			if(UAnimNextDataInterfaceEntry* AssetEntry = Item->WeakEntry.Get())
+			if(UAnimNextSharedVariablesEntry* AssetEntry = Item->WeakEntry.Get())
 			{
 				bool bIsDirty = false;
 				const UPackage* Package = AssetEntry->GetPackage();
@@ -164,7 +164,7 @@ class SVariablesOutlinerDataInterfaceLabel : FSceneOutlinerCommonLabelData, publ
 	TSharedPtr<SInlineEditableTextBlock> TextBlock;
 };
 
-FVariablesOutlinerDataInterfaceItem::FVariablesOutlinerDataInterfaceItem(UAnimNextDataInterfaceEntry* InEntry)
+FVariablesOutlinerDataInterfaceItem::FVariablesOutlinerDataInterfaceItem(UAnimNextSharedVariablesEntry* InEntry)
 	: ISceneOutlinerTreeItem(FVariablesOutlinerDataInterfaceItem::Type)
 	, WeakEntry(InEntry)
 {
diff --git a/Engine/Plugins/Experimental/AnimNext/Source/AnimNextEditor/Private/Variables/Outliner/VariablesOutlinerDataInterfaceItem.h b/Engine/Plugins/Experimental/AnimNext/Source/AnimNextEditor/Private/Variables/Outliner/VariablesOutlinerDataInterfaceItem.h
--- a/Engine/Plugins/Experimental/AnimNext/Source/AnimNextEditor/Private/Variables/Outliner/VariablesOutlinerDataInterfaceItem.h
+++ b/Engine/Plugins/Experimental/AnimNext/Source/AnimNextEditor/Private/Variables/Outliner/VariablesOutlinerDataInterfaceItem.h
@@ -4,7 +4,7 @@
 
 #include "ISceneOutlinerTreeItem.h"
 
-class UAnimNextDataInterfaceEntry;
+class UAnimNextSharedVariablesEntry;
 
 namespace UE::AnimNext::Editor
 {
@@ -15,7 +15,7 @@ struct FVariablesOutlinerDataInterfaceItem : ISceneOutlinerTreeItem
 {
 	static const FSceneOutlinerTreeItemType Type;
 
-	FVariablesOutlinerDataInterfaceItem(UAnimNextDataInterfaceEntry* InEntry);
+	FVariablesOutlinerDataInterfaceItem(UAnimNextSharedVariablesEntry* InEntry);
 
 	// Begin ISceneOutlinerTreeItem overrides
 	virtual bool IsValid() const override;
@@ -27,7 +27,7 @@ struct FVariablesOutlinerDataInterfaceItem : ISceneOutlinerTreeItem
 	// End ISceneOutlinerTreeItem overrides
 
 	// Ptr to the underlying entry
-	TWeakObjectPtr<UAnimNextDataInterfaceEntry> WeakEntry;
+	TWeakObjectPtr<UAnimNextSharedVariablesEntry> WeakEntry;
 };
 
 }
diff --git a/Engine/Plugins/Experimental/AnimNext/Source/AnimNextEditor/Private/Variables/Outliner/VariablesOutlinerEntryItem.cpp b/Engine/Plugins/Experimental/AnimNext/Source/AnimNextEditor/Private/Variables/Outliner/VariablesOutlinerEntryItem.cpp
--- a/Engine/Plugins/Experimental/AnimNext/Source/AnimNextEditor/Private/Variables/Outliner/VariablesOutlinerEntryItem.cpp
+++ b/Engine/Plugins/Experimental/AnimNext/Source/AnimNextEditor/Private/Variables/Outliner/VariablesOutlinerEntryItem.cpp
@@ -51,6 +51,7 @@ class SVariablesOutlinerEntrylabel : FSceneOutlinerCommonLabelData, public SComp
 			.Padding(0.0f, 2.0f)
 			[
 				SAssignNew(TextBlock, SInlineEditableTextBlock)
+				.IsReadOnly(this, &SVariablesOutlinerEntrylabel::CanRename)
 				.Text(this, &SVariablesOutlinerEntrylabel::GetDisplayText)
 				.HighlightText(SceneOutliner.GetFilterHighlightText())
 				.ColorAndOpacity(this, &SVariablesOutlinerEntrylabel::GetForegroundColor)
@@ -138,6 +139,15 @@ class SVariablesOutlinerEntrylabel : FSceneOutlinerCommonLabelData, public SComp
 		return false;
 	}
 
+	bool CanRename() const
+	{
+		if (const TSharedPtr<FVariablesOutlinerEntryItem> Item = TreeItem.Pin())
+		{
+			return Item->PropertyPath.IsPathToFieldEmpty() && Item->WeakEntry.Get() != nullptr;
+		}
+		return false;
+	}
+	
 	virtual FSlateColor GetForegroundColor() const override
 	{
 		const TOptional<FLinearColor> BaseColor = FSceneOutlinerCommonLabelData::GetForegroundColor(*TreeItem.Pin());
@@ -154,25 +164,41 @@ FVariablesOutlinerEntryItem::FVariablesOutlinerEntryItem(UAnimNextVariableEntry*
 {
 }
 
+FVariablesOutlinerEntryItem::FVariablesOutlinerEntryItem(const FProperty* InProperty)
+	: ISceneOutlinerTreeItem(FVariablesOutlinerEntryItem::Type)
+	, PropertyPath(InProperty)
+{
+}
+
 bool FVariablesOutlinerEntryItem::IsValid() const
 {
-	return WeakEntry.Get() != nullptr;
+	return WeakEntry.Get() != nullptr || PropertyPath.Get() != nullptr;
 }
 
 FSceneOutlinerTreeItemID FVariablesOutlinerEntryItem::GetID() const
 {
-	return HashCombine(GetTypeHash(WeakEntry.Get()), GetTypeHash(WeakDataInterfaceEntry.Get()));
+	if (UAnimNextVariableEntry* Entry = WeakEntry.Get())
+	{
+		return HashCombine(GetTypeHash(Entry), GetTypeHash(WeakDataInterfaceEntry.Get()));
+	}
+	else
+	{
+		return HashCombine(GetTypeHash(PropertyPath), GetTypeHash(WeakDataInterfaceEntry.Get()));
+	}
 }
 
 FString FVariablesOutlinerEntryItem::GetDisplayString() const
 {
-	UAnimNextRigVMAssetEntry* Entry = WeakEntry.Get();
-	if(Entry == nullptr)
+	if(UAnimNextRigVMAssetEntry* Entry = WeakEntry.Get())
 	{
-		return FString();
+		return Entry->GetDisplayName().ToString();
+	}
+	else if (const FProperty* Property = PropertyPath.Get())
+	{
+		return Property->GetName();
 	}
 
-	return Entry->GetDisplayName().ToString();
+	return FString();
 }
 
 TSharedRef<SWidget> FVariablesOutlinerEntryItem::GenerateLabelWidget(ISceneOutliner& Outliner, const STableRow<FSceneOutlinerTreeItemPtr>& InRow)
@@ -184,25 +210,25 @@ TSharedRef<SWidget> FVariablesOutlinerEntryItem::GenerateLabelWidget(ISceneOutli
 
 FString FVariablesOutlinerEntryItem::GetPackageName() const
 {
-	UAnimNextRigVMAssetEntry* Entry = WeakEntry.Get();
-	if(Entry == nullptr)
+	if(UAnimNextRigVMAssetEntry* Entry = WeakEntry.Get())
 	{
-		return ISceneOutlinerTreeItem::GetPackageName();
+		return Entry->GetPackage()->GetName();
 	}
-
-	return Entry->GetPackage()->GetName();
+	else if (const FProperty* Property = PropertyPath.Get())
+	{
+		return Property->GetOwner<UScriptStruct>()->GetPackage()->GetName();
+	}
+	
+	return ISceneOutlinerTreeItem::GetPackageName();
 }
 
 void FVariablesOutlinerEntryItem::Rename(const FText& InNewName)
 {
-	UAnimNextRigVMAssetEntry* Entry = WeakEntry.Get();
-	if(Entry == nullptr)
+	if(UAnimNextRigVMAssetEntry* Entry = WeakEntry.Get())
 	{
-		return;
+		FScopedTransaction Transaction(LOCTEXT("RenameVariable", "Rename variable"));
+		Entry->SetEntryName(FName(*InNewName.ToString()));
 	}
-
-	FScopedTransaction Transaction(LOCTEXT("RenameVariable", "Rename variable"));
-	Entry->SetEntryName(FName(*InNewName.ToString()));
 }
 
 bool FVariablesOutlinerEntryItem::ValidateName(const FText& InNewName, FText& OutErrorMessage) const
diff --git a/Engine/Plugins/Experimental/AnimNext/Source/AnimNextEditor/Private/Variables/Outliner/VariablesOutlinerEntryItem.h b/Engine/Plugins/Experimental/AnimNext/Source/AnimNextEditor/Private/Variables/Outliner/VariablesOutlinerEntryItem.h
--- a/Engine/Plugins/Experimental/AnimNext/Source/AnimNextEditor/Private/Variables/Outliner/VariablesOutlinerEntryItem.h
+++ b/Engine/Plugins/Experimental/AnimNext/Source/AnimNextEditor/Private/Variables/Outliner/VariablesOutlinerEntryItem.h
@@ -3,9 +3,8 @@
 #pragma once
 
 #include "ISceneOutlinerTreeItem.h"
-#include "Variables/AnimNextVariableEntryProxy.h"
 
-class UAnimNextDataInterfaceEntry;
+class UAnimNextSharedVariablesEntry;
 class UAnimNextVariableEntry;
 
 namespace UE::AnimNext::Editor
@@ -16,6 +15,7 @@ struct FVariablesOutlinerEntryItem : ISceneOutlinerTreeItem
 	static const FSceneOutlinerTreeItemType Type;
 
 	FVariablesOutlinerEntryItem(UAnimNextVariableEntry* InEntry);
+	FVariablesOutlinerEntryItem(const FProperty* InProperty);
 
 	// Begin ISceneOutlinerTreeItem overrides
 	virtual bool IsValid() const override;
@@ -32,14 +32,14 @@ struct FVariablesOutlinerEntryItem : ISceneOutlinerTreeItem
 	// Validates the new item name
 	bool ValidateName(const FText& InNewName, FText& OutErrorMessage) const;
 
-	// Ptr to the underlying entry
+	// Ptr to the underlying entry if this variable came from an asset
 	TWeakObjectPtr<UAnimNextVariableEntry> WeakEntry;
 
 	// The data interface entry this entry is from, if any
-	TWeakObjectPtr<UAnimNextDataInterfaceEntry> WeakDataInterfaceEntry;
+	TWeakObjectPtr<UAnimNextSharedVariablesEntry> WeakDataInterfaceEntry;
 
-	// Proxy entry used for details panel editing of variables in implemented data interfaces
-	TStrongObjectPtr<UAnimNextVariableEntryProxy> ProxyEntry;
+	// Ptr to the property if this variable came from a struct
+	TFieldPath<const FProperty> PropertyPath;
 };
 
 }
diff --git a/Engine/Plugins/Experimental/AnimNext/Source/AnimNextEditor/Private/Variables/Outliner/VariablesOutlinerHierarchy.cpp b/Engine/Plugins/Experimental/AnimNext/Source/AnimNextEditor/Private/Variables/Outliner/VariablesOutlinerHierarchy.cpp
--- a/Engine/Plugins/Experimental/AnimNext/Source/AnimNextEditor/Private/Variables/Outliner/VariablesOutlinerHierarchy.cpp
+++ b/Engine/Plugins/Experimental/AnimNext/Source/AnimNextEditor/Private/Variables/Outliner/VariablesOutlinerHierarchy.cpp
@@ -10,11 +10,11 @@
 #include "VariablesOutlinerMode.h"
 #include "WorkspaceAssetRegistryInfo.h"
 #include "VariablesOutlinerEntryItem.h"
-#include "Entries/AnimNextDataInterfaceEntry.h"
+#include "Entries/AnimNextSharedVariablesEntry.h"
 #include "Entries/AnimNextRigVMAssetEntry.h"
 #include "Entries/AnimNextVariableEntry.h"
 #include "Variables/SVariablesView.h"
-#include "DataInterface/AnimNextDataInterface.h"
+#include "Variables/AnimNextSharedVariables.h"
 
 namespace UE::AnimNext::Editor
 {
@@ -55,35 +55,57 @@ void FVariablesOutlinerHierarchy::CreateItems(TArray<FSceneOutlinerTreeItemPtr>&
 			return true;
 		});
 
-		EditorData->ForEachEntryOfType<UAnimNextDataInterfaceEntry>([this, &OutItems](UAnimNextDataInterfaceEntry* InInterfaceEntry)
+		EditorData->ForEachEntryOfType<UAnimNextSharedVariablesEntry>([this, &OutItems](UAnimNextSharedVariablesEntry* InInterfaceEntry)
 		{
-			if (FSceneOutlinerTreeItemPtr Item = Mode->CreateItemFor<FVariablesOutlinerDataInterfaceItem>(CastChecked<UAnimNextDataInterfaceEntry>(InInterfaceEntry)))
+			if (FSceneOutlinerTreeItemPtr Item = Mode->CreateItemFor<FVariablesOutlinerDataInterfaceItem>(CastChecked<UAnimNextSharedVariablesEntry>(InInterfaceEntry)))
 			{
 				OutItems.Add(Item);
 			}
 
-			auto AddDataInterface = [this, &OutItems](UAnimNextDataInterfaceEntry* InDataInterfaceEntry, UAnimNextDataInterfaceEntry* InRootDataInterfaceEntry, auto& InAddDataInterface) -> void
+			auto AddDataInterface = [this, &OutItems](UAnimNextSharedVariablesEntry* InDataInterfaceEntry, UAnimNextSharedVariablesEntry* InRootDataInterfaceEntry, auto& InAddDataInterface) -> void
 			{
-				UAnimNextRigVMAssetEditorData* InterfaceEditorData = UncookedOnly::FUtils::GetEditorData<UAnimNextRigVMAssetEditorData>(InDataInterfaceEntry->DataInterface.Get());
-				InterfaceEditorData->ForEachEntryOfType<UAnimNextVariableEntry>([this, &OutItems, InRootDataInterfaceEntry](UAnimNextVariableEntry* InVariable)
+				switch (InDataInterfaceEntry->GetType())
 				{
-					if(InVariable->GetExportAccessSpecifier() == EAnimNextExportAccessSpecifier::Public)
+				case EAnimNextSharedVariablesType::Asset:
 					{
-						if (FSceneOutlinerTreeItemPtr Item = Mode->CreateItemFor<FVariablesOutlinerEntryItem>(CastChecked<UAnimNextVariableEntry>(InVariable)))
+						UAnimNextRigVMAssetEditorData* InterfaceEditorData = UncookedOnly::FUtils::GetEditorData<UAnimNextRigVMAssetEditorData>(InDataInterfaceEntry->GetAsset());
+						InterfaceEditorData->ForEachEntryOfType<UAnimNextVariableEntry>([this, &OutItems, InRootDataInterfaceEntry](UAnimNextVariableEntry* InVariable)
 						{
-							FVariablesOutlinerEntryItem* EntryItem = Item->CastTo<FVariablesOutlinerEntryItem>();
-							EntryItem->WeakDataInterfaceEntry = InRootDataInterfaceEntry;
-							OutItems.Add(Item);
+							if(InVariable->GetExportAccessSpecifier() == EAnimNextExportAccessSpecifier::Public)
+							{
+								if (FSceneOutlinerTreeItemPtr Item = Mode->CreateItemFor<FVariablesOutlinerEntryItem>(CastChecked<UAnimNextVariableEntry>(InVariable)))
+								{
+									FVariablesOutlinerEntryItem* EntryItem = Item->CastTo<FVariablesOutlinerEntryItem>();
+									EntryItem->WeakDataInterfaceEntry = InRootDataInterfaceEntry;
+									OutItems.Add(Item);
+								}
+							}
+							return true;
+						});
+
+						InterfaceEditorData->ForEachEntryOfType<UAnimNextSharedVariablesEntry>([this, &OutItems, &InAddDataInterface, InRootDataInterfaceEntry](UAnimNextSharedVariablesEntry* InSubInterfaceEntry)
+						{
+							InAddDataInterface(InSubInterfaceEntry, InRootDataInterfaceEntry, InAddDataInterface);
+							return true;
+						});
+						break;
+					}
+				case EAnimNextSharedVariablesType::Struct:
+					{
+						const UScriptStruct* Struct = InDataInterfaceEntry->GetStruct();
+						for (TFieldIterator<FProperty> It(Struct); It; ++It)
+						{
+							const FProperty* Property = *It;
+							if (FSceneOutlinerTreeItemPtr Item = Mode->CreateItemFor<FVariablesOutlinerEntryItem>(Property))
+							{
+								FVariablesOutlinerEntryItem* EntryItem = Item->CastTo<FVariablesOutlinerEntryItem>();
+								EntryItem->WeakDataInterfaceEntry = InRootDataInterfaceEntry;
+								OutItems.Add(Item);
+							}
 						}
+						break;
 					}
-					return true;
-				});
-
-				InterfaceEditorData->ForEachEntryOfType<UAnimNextDataInterfaceEntry>([this, &OutItems, &InAddDataInterface, InRootDataInterfaceEntry](UAnimNextDataInterfaceEntry* InSubInterfaceEntry)
-				{
-					InAddDataInterface(InSubInterfaceEntry, InRootDataInterfaceEntry, InAddDataInterface);
-					return true;
-				});
+				}
 			};
 
 			AddDataInterface(InInterfaceEntry, InInterfaceEntry, AddDataInterface);
@@ -97,11 +119,11 @@ FSceneOutlinerTreeItemPtr FVariablesOutlinerHierarchy::FindOrCreateParentItem(co
 {
 	if (const FVariablesOutlinerEntryItem* EntryItem = Item.CastTo<FVariablesOutlinerEntryItem>())
 	{
-		UAnimNextDataInterfaceEntry* DataInterfaceEntry = EntryItem->WeakDataInterfaceEntry.Get();
-		if(DataInterfaceEntry != nullptr)
+		UAnimNextSharedVariablesEntry* SharedVariablesEntry = EntryItem->WeakDataInterfaceEntry.Get();
+		if(SharedVariablesEntry != nullptr)
 		{
 			// Add as part of an imported data interface, so use that as parent
-			TSoftObjectPtr<UAnimNextRigVMAssetEntry> SoftObjectPtr(DataInterfaceEntry); 
+			TSoftObjectPtr<UAnimNextRigVMAssetEntry> SoftObjectPtr(SharedVariablesEntry); 
 			uint32 Hash = GetTypeHash(SoftObjectPtr);
 
 			if (const FSceneOutlinerTreeItemPtr* ParentItem = Items.Find(Hash))
diff --git a/Engine/Plugins/Experimental/AnimNext/Source/AnimNextEditor/Private/Variables/Outliner/VariablesOutlinerMode.cpp b/Engine/Plugins/Experimental/AnimNext/Source/AnimNextEditor/Private/Variables/Outliner/VariablesOutlinerMode.cpp
--- a/Engine/Plugins/Experimental/AnimNext/Source/AnimNextEditor/Private/Variables/Outliner/VariablesOutlinerMode.cpp
+++ b/Engine/Plugins/Experimental/AnimNext/Source/AnimNextEditor/Private/Variables/Outliner/VariablesOutlinerMode.cpp
@@ -25,8 +25,8 @@
 #include "Common/GraphEditorSchemaActions.h"
 #include "Entries/AnimNextVariableEntry.h"
 #include "Variables/SVariablesView.h"
-#include "Entries/AnimNextDataInterfaceEntry.h"
-#include "Variables/AnimNextVariableEntryProxy.h"
+#include "Entries/AnimNextSharedVariablesEntry.h"
+#include "Variables/AnimNextSharedVariables.h"
 
 #define LOCTEXT_NAMESPACE "VariablesOutlinerMode"
 
@@ -162,28 +162,14 @@ void FVariablesOutlinerMode::HandleItemSelection(const FSceneOutlinerItemSelecti
 			{
 				if(UAnimNextVariableEntry* VariableEntry = VariablesItem->WeakEntry.Get())
 				{
-					if(UAnimNextDataInterfaceEntry* DataInterfaceEntry = VariablesItem->WeakDataInterfaceEntry.Get())
-					{
-						// Create proxy object to display in the details panel
-						if(!VariablesItem->ProxyEntry.IsValid())
-						{
-							VariablesItem->ProxyEntry = TStrongObjectPtr(NewObject<UAnimNextVariableEntryProxy>(GetTransientPackage(), NAME_None, RF_Transient));
-						}
-						VariablesItem->ProxyEntry.Get()->VariableEntry = VariableEntry;
-						VariablesItem->ProxyEntry.Get()->DataInterfaceEntry = DataInterfaceEntry;
-						EntriesToShow.Add(VariablesItem->ProxyEntry.Get());
-					}
-					else
-					{
-						EntriesToShow.Add(VariableEntry);
-					}
+					EntriesToShow.Add(VariableEntry);
 				}
 			}
 			else if(const FVariablesOutlinerDataInterfaceItem* DataInterfaceItem = Item->CastTo<FVariablesOutlinerDataInterfaceItem>())
 			{
-				if(UAnimNextDataInterfaceEntry* DataInterfaceEntry = DataInterfaceItem->WeakEntry.Get())
+				if(UAnimNextSharedVariablesEntry* SharedVariablesEntry = DataInterfaceItem->WeakEntry.Get())
 				{
-					EntriesToShow.Add(DataInterfaceEntry);
+					EntriesToShow.Add(SharedVariablesEntry);
 				}
 			}
 		}
@@ -280,8 +266,8 @@ void FVariablesOutlinerMode::Delete()
 		if (const FVariablesOutlinerEntryItem* VariablesItem = Item->CastTo<FVariablesOutlinerEntryItem>())
 		{
 			UAnimNextVariableEntry* VariableEntry = VariablesItem->WeakEntry.Get();
-			UAnimNextDataInterfaceEntry* DataInterfaceEntry = VariablesItem->WeakDataInterfaceEntry.Get();
-			if(VariableEntry == nullptr || DataInterfaceEntry != nullptr)	// Cant delete variables in other data interfaces
+			UAnimNextSharedVariablesEntry* SharedVariablesEntry = VariablesItem->WeakDataInterfaceEntry.Get();
+			if(VariableEntry == nullptr || SharedVariablesEntry != nullptr)	// Cant delete variables in other data interfaces
 			{
 				continue;
 			}
@@ -298,20 +284,20 @@ void FVariablesOutlinerMode::Delete()
 		}
 		else if (const FVariablesOutlinerDataInterfaceItem* DataInterfaceItem = Item->CastTo<FVariablesOutlinerDataInterfaceItem>())
 		{
-			UAnimNextDataInterfaceEntry* DataInterfaceEntry = DataInterfaceItem->WeakEntry.Get();
-			if(DataInterfaceEntry == nullptr)
+			UAnimNextSharedVariablesEntry* SharedVariablesEntry = DataInterfaceItem->WeakEntry.Get();
+			if(SharedVariablesEntry == nullptr)
 			{
 				continue;
 			}
 
-			UAnimNextRigVMAssetEditorData* EditorData = DataInterfaceEntry->GetTypedOuter<UAnimNextRigVMAssetEditorData>();
+			UAnimNextRigVMAssetEditorData* EditorData = SharedVariablesEntry->GetTypedOuter<UAnimNextRigVMAssetEditorData>();
 			if(EditorData == nullptr)
 			{
 				continue;
 			}
 
 			TArray<UAnimNextRigVMAssetEntry*>& EntriesToDelete = EntriesToDeletePerAsset.FindOrAdd(EditorData);
-			EntriesToDelete.Add(DataInterfaceEntry);
+			EntriesToDelete.Add(SharedVariablesEntry);
 			NumEntries++;
 		}
 	}
@@ -343,21 +329,36 @@ TSharedPtr<FDragDropOperation> FVariablesOutlinerMode::CreateDragDropOperation(c
 	{
 		return nullptr;
 	}
-	
-	UAnimNextVariableEntry* Entry = Cast<UAnimNextVariableEntry>(VariableItem->WeakEntry.Get());
-	if(Entry == nullptr)
+
+	if(UAnimNextVariableEntry* Entry = Cast<UAnimNextVariableEntry>(VariableItem->WeakEntry.Get()))
 	{
-		return nullptr;
-	}
+		const UAnimNextSharedVariables* Asset = nullptr;
+		if (VariableItem->WeakDataInterfaceEntry.Get())
+		{
+			Asset = VariableItem->WeakDataInterfaceEntry->GetAsset();
+		}
+		else
+		{
+			Asset = Entry->GetTypedOuter<UAnimNextSharedVariables>();
+		}
 
-	URigVMHost* RigVMHost = VariableItem->WeakDataInterfaceEntry.IsValid() ? VariableItem->WeakDataInterfaceEntry.Get()->GetTypedOuter<URigVMHost>() : Entry->GetTypedOuter<URigVMHost>();
-	if(RigVMHost == nullptr)
+		if (Asset)
+		{
+			TSharedPtr<FAnimNextSchemaAction_Variable> Action = MakeShared<FAnimNextSchemaAction_Variable>(Entry->GetVariableName(), Asset, Entry->GetType(), FAnimNextSchemaAction_Variable::EVariableAccessorChoice::Deferred);
+			return FVariableDragDropOp::New(Action);
+		}
+	}
+	else if (const FProperty* Property = VariableItem->PropertyPath.Get())
 	{
-		return nullptr;
+		const UScriptStruct* Struct = Property->GetOwner<UScriptStruct>();
+		if (Struct)
+		{
+			TSharedPtr<FAnimNextSchemaAction_Variable> Action = MakeShared<FAnimNextSchemaAction_Variable>(Property->GetFName(), Struct, FAnimNextParamType::FromProperty(Property), FAnimNextSchemaAction_Variable::EVariableAccessorChoice::Deferred);
+			return FVariableDragDropOp::New(Action);
+		}
 	}
 
-	TSharedPtr<FAnimNextSchemaAction_Variable> Action = MakeShared<FAnimNextSchemaAction_Variable>(Entry->GetVariableName(), Entry->GetType(), FAnimNextSchemaAction_Variable::EVariableAccessorChoice::Deferred);
-	return FVariableDragDropOp::New(Action);
+	return nullptr;
 }
 
 }
diff --git a/Engine/Plugins/Experimental/AnimNext/Source/AnimNextEditor/Private/Variables/Outliner/VariablesOutlinerMode.h b/Engine/Plugins/Experimental/AnimNext/Source/AnimNextEditor/Private/Variables/Outliner/VariablesOutlinerMode.h
--- a/Engine/Plugins/Experimental/AnimNext/Source/AnimNextEditor/Private/Variables/Outliner/VariablesOutlinerMode.h
+++ b/Engine/Plugins/Experimental/AnimNext/Source/AnimNextEditor/Private/Variables/Outliner/VariablesOutlinerMode.h
@@ -6,7 +6,7 @@
 #include "ISceneOutlinerMode.h"
 
 class UAnimNextRigVMAssetEditorData;
-
+class FVariablesOutlinerMode;
 namespace UE::Workspace
 {
 	class IWorkspaceEditor;
diff --git a/Engine/Plugins/Experimental/AnimNext/Source/AnimNextEditor/Private/Variables/SAddVariablesDialog.cpp b/Engine/Plugins/Experimental/AnimNext/Source/AnimNextEditor/Private/Variables/SAddVariablesDialog.cpp
--- a/Engine/Plugins/Experimental/AnimNext/Source/AnimNextEditor/Private/Variables/SAddVariablesDialog.cpp
+++ b/Engine/Plugins/Experimental/AnimNext/Source/AnimNextEditor/Private/Variables/SAddVariablesDialog.cpp
@@ -22,10 +22,12 @@
 #include "Widgets/Text/SInlineEditableTextBlock.h"
 #include "SSimpleButton.h"
 #include "SSimpleComboButton.h"
+#include "StructViewerFilter.h"
+#include "StructViewerModule.h"
 #include "ToolMenus.h"
-#include "DataInterface/AnimNextDataInterface.h"
-#include "DataInterface/AnimNextDataInterface_EditorData.h"
-#include "Entries/AnimNextDataInterfaceEntry.h"
+#include "Variables/AnimNextSharedVariables.h"
+#include "Variables/AnimNextSharedVariables_EditorData.h"
+#include "Entries/AnimNextSharedVariablesEntry.h"
 #include "Entries/AnimNextRigVMAssetEntry.h"
 #include "Entries/AnimNextVariableEntry.h"
 #include "String/ParseTokens.h"
@@ -87,7 +89,7 @@ bool SAddVariablesDialog::FVariableToAddEntry::IsValid(FText& OutReason) const
 					return false;
 				}
 				break;
-			case EEntryType::DataInterface:
+			case EEntryType::SharedVariables:
 				for(const TSharedRef<FEntry>& SubEntry : Entry->Children)
 				{
 					if(&SubEntry.Get() == this)
@@ -110,11 +112,12 @@ bool SAddVariablesDialog::FVariableToAddEntry::IsValid(FText& OutReason) const
 	return true; 
 }
 
-bool SAddVariablesDialog::FDataInterfaceToAddEntry::IsValid(FText& OutReason) const
+bool SAddVariablesDialog::FSharedVariablesToAddEntry::IsValid(FText& OutReason) const
 {
-	if(DataInterface == nullptr)
+	if( (SharedVariables.IsType<UAnimNextSharedVariables*>() && SharedVariables.Get<UAnimNextSharedVariables*>() == nullptr) ||
+		(SharedVariables.IsType<const UScriptStruct*>() && SharedVariables.Get<const UScriptStruct*>() == nullptr))
 	{
-		OutReason = LOCTEXT("InvalidDataInterface", "Invalid Data Interface");
+		OutReason = LOCTEXT("InvalidSharedVariables", "Invalid Shared Variables");
 		return false;
 	}
 
@@ -166,8 +169,8 @@ void SAddVariablesDialog::Construct(const FArguments& InArgs, const TArray<UAnim
 					[
 						SNew(SSimpleComboButton)
 						.Visibility(InArgs._AllowMultiple ? EVisibility::Visible : EVisibility::Collapsed)
-						.Text(LOCTEXT("AddDataInterfaceButton", "Add Data Interface"))
-						.ToolTipText(LOCTEXT("AddDataInterfaceButtonTooltip", "Select a new data interface for adding."))
+						.Text(LOCTEXT("AddSharedVariablesButton", "Add Shared Variables"))
+						.ToolTipText(LOCTEXT("AddSharedVariablesButtonTooltip", "Add a set of shared variables from another asset."))
 						.Icon(FAppStyle::Get().GetBrush("Icons.Plus"))
 						.OnGetMenuContent_Lambda([this]()
 						{
@@ -175,14 +178,14 @@ void SAddVariablesDialog::Construct(const FArguments& InArgs, const TArray<UAnim
 
 							FAssetPickerConfig AssetPickerConfig;
 							AssetPickerConfig.Filter.bRecursiveClasses = true;
-							AssetPickerConfig.Filter.ClassPaths.Add(UAnimNextDataInterface::StaticClass()->GetClassPathName());
+							AssetPickerConfig.Filter.ClassPaths.Add(UAnimNextSharedVariables::StaticClass()->GetClassPathName());
 							AssetPickerConfig.InitialAssetViewType = EAssetViewType::List;
 							AssetPickerConfig.OnAssetSelected = FOnAssetSelected::CreateLambda([this](const FAssetData& InAssetData)
 							{
 								FSlateApplication::Get().DismissAllMenus();
-								if(UAnimNextDataInterface* DataInterface = Cast<UAnimNextDataInterface>(InAssetData.GetAsset()))
+								if(UAnimNextSharedVariables* SharedVariables = Cast<UAnimNextSharedVariables>(InAssetData.GetAsset()))
 								{
-									AddDataInterface(DataInterface);
+									AddSharedVariables(SharedVariables);
 								}
 							});
 							AssetPickerConfig.OnShouldFilterAsset = FOnShouldFilterAsset::CreateLambda([](const FAssetData& InAssetData)
@@ -204,6 +207,60 @@ void SAddVariablesDialog::Construct(const FArguments& InArgs, const TArray<UAnim
 								];
 						})
 					]
+					+SHorizontalBox::Slot()
+					.AutoWidth()
+					.Padding(0.3f, 0.0f)
+					[
+						SNew(SSimpleComboButton)
+						.Visibility(InArgs._AllowMultiple ? EVisibility::Visible : EVisibility::Collapsed)
+						.Text(LOCTEXT("AddNativeSharedVariablesButton", "Add Native Shared Variables"))
+						.ToolTipText(LOCTEXT("AddNativeSharedVariablesButtonTooltip", "Add a set of shared variables from a native struct."))
+						.Icon(FAppStyle::Get().GetBrush("Icons.Plus"))
+						.OnGetMenuContent_Lambda([this]()
+						{
+							class FStructFilter : public IStructViewerFilter
+							{
+							public:
+								virtual bool IsStructAllowed(const FStructViewerInitializationOptions& InInitOptions, const UScriptStruct* InStruct, TSharedRef<FStructViewerFilterFuncs> InFilterFuncs) override
+								{
+									if (InStruct->IsA<UUserDefinedStruct>())
+									{
+										return false;
+									}
+
+									if (!InStruct->HasMetaData(TEXT("BlueprintType")) || InStruct->HasMetaData(TEXT("Hidden")))
+									{
+										return false;
+									}
+
+									return true;
+								}
+
+								virtual bool IsUnloadedStructAllowed(const FStructViewerInitializationOptions& InInitOptions, const FSoftObjectPath& InStructPath, TSharedRef<class FStructViewerFilterFuncs> InFilterFuncs)
+								{
+									return false;
+								};
+							};
+							
+							FStructViewerModule& StructViewerModule = FModuleManager::LoadModuleChecked<FStructViewerModule>("StructViewer");
+
+							FStructViewerInitializationOptions Options;
+							Options.Mode = EStructViewerMode::StructPicker;
+							Options.StructFilter = MakeShared<FStructFilter>();
+
+							return SNew(SBox)
+								.WidthOverride(300.0f)
+								.HeightOverride(400.0f)
+								[
+									StructViewerModule.CreateStructViewer(Options, FOnStructPicked::CreateLambda([this](const UScriptStruct* InStruct)
+									{
+										FSlateApplication::Get().DismissAllMenus();
+
+										AddSharedVariablesStruct(InStruct);
+									}))
+								];
+						})
+					]
 				]
 				+SVerticalBox::Slot()
 				.FillHeight(1.0f)
@@ -361,11 +418,11 @@ void SAddVariablesDialog::GetPendingNamesRecursive(UAnimNextRigVMAssetEditorData
 				OutPendingNames.Add(Entry->GetEntryName());
 			}
 		}
-		else if(UAnimNextDataInterfaceEntry* DataInterfaceEntry = Cast<UAnimNextDataInterfaceEntry>(Entry))
+		else if(UAnimNextSharedVariablesEntry* DataInterfaceEntry = Cast<UAnimNextSharedVariablesEntry>(Entry))
 		{
-			if(DataInterfaceEntry->DataInterface)
+			if(DataInterfaceEntry->GetAsset())
 			{
-				UAnimNextDataInterface_EditorData* EditorData = UE::AnimNext::UncookedOnly::FUtils::GetEditorData<UAnimNextDataInterface_EditorData>(DataInterfaceEntry->DataInterface.Get());
+				UAnimNextSharedVariables_EditorData* EditorData = UE::AnimNext::UncookedOnly::FUtils::GetEditorData<UAnimNextSharedVariables_EditorData>(DataInterfaceEntry->GetAsset());
 				GetPendingNamesRecursive(EditorData, OutPendingNames);
 			}
 		}
@@ -393,7 +450,7 @@ void SAddVariablesDialog::AddEntry(const FAnimNextParamType& InParamType)
 		case EEntryType::Variable:
 			PendingNames.Add(StaticCastSharedRef<FVariableToAddEntry>(QueuedAdd)->Name);
 			break;
-		case EEntryType::DataInterface:
+		case EEntryType::SharedVariables:
 			for(const TSharedRef<FEntry>& SubEntry : QueuedAdd->Children)
 			{
 				check(SubEntry->EntryType == EEntryType::Variable);
@@ -411,10 +468,9 @@ void SAddVariablesDialog::AddEntry(const FAnimNextParamType& InParamType)
 	RefreshEntries();
 }
 
-void SAddVariablesDialog::AddDataInterface(UAnimNextDataInterface* InDataInterface)
+void SAddVariablesDialog::AddSharedVariables(UAnimNextSharedVariables* InSharedVariables)
 {
-	
-	TSharedRef<FDataInterfaceToAddEntry> NewEntry = MakeShared<FDataInterfaceToAddEntry>(InDataInterface, SharedThis(this));
+	TSharedRef<FSharedVariablesToAddEntry> NewEntry = MakeShared<FSharedVariablesToAddEntry>(InSharedVariables, SharedThis(this));
 
 	auto AddVariable = [this, &NewEntry](UAnimNextVariableEntry* InVariableEntry)
 	{
@@ -427,28 +483,59 @@ void SAddVariablesDialog::AddDataInterface(UAnimNextDataInterface* InDataInterfa
 		}
 	};
 
-	auto AddDataInterfaceInternal = [this, &AddVariable](UAnimNextDataInterface* InDataInterfaceToAdd, auto& InAddDataInterfaceInternal) -> void
+	auto AddSharedVariablesInternal = [this, &AddVariable](const UAnimNextSharedVariables* InSharedVariablesToAdd, auto& InAddSharedVariablesInternal) -> void
 	{
 		// Add the child entries
-		UAnimNextDataInterface_EditorData* EditorData = UncookedOnly::FUtils::GetEditorData<UAnimNextDataInterface_EditorData>(InDataInterfaceToAdd);
+		UAnimNextSharedVariables_EditorData* EditorData = UncookedOnly::FUtils::GetEditorData<UAnimNextSharedVariables_EditorData>(InSharedVariablesToAdd);
 		for(UAnimNextRigVMAssetEntry* AssetEntry : EditorData->Entries)
 		{
 			if(UAnimNextVariableEntry* VariableEntry = Cast<UAnimNextVariableEntry>(AssetEntry))
 			{
 				AddVariable(VariableEntry);
 			}
-			else if(UAnimNextDataInterfaceEntry* DataInterfaceEntry = Cast<UAnimNextDataInterfaceEntry>(AssetEntry))
+			else if(UAnimNextSharedVariablesEntry* SharedVariablesEntry = Cast<UAnimNextSharedVariablesEntry>(AssetEntry))
 			{
-				if(DataInterfaceEntry->GetDataInterface())
+				if(SharedVariablesEntry->GetAsset())
 				{
-					InAddDataInterfaceInternal(DataInterfaceEntry->GetDataInterface(), InAddDataInterfaceInternal);
+					InAddSharedVariablesInternal(SharedVariablesEntry->GetAsset(), InAddSharedVariablesInternal);
 				}
 			}
 		}
 	};
 
-	// Add the selected data interface and its recursive dependents
-	AddDataInterfaceInternal(InDataInterface, AddDataInterfaceInternal);
+	// Add the selected shared variables and its recursive dependents
+	AddSharedVariablesInternal(InSharedVariables, AddSharedVariablesInternal);
+
+	if(NewEntry->Children.Num() > 0)
+	{
+		RootEntries.Add(NewEntry);
+		RefreshEntries();
+	}
+}
+
+void SAddVariablesDialog::AddSharedVariablesStruct(const UScriptStruct* InStruct)
+{
+	TSharedRef<FSharedVariablesToAddEntry> NewEntry = MakeShared<FSharedVariablesToAddEntry>(InStruct, SharedThis(this));
+
+	auto AddVariable = [this, &NewEntry](FProperty* InProperty)
+	{
+		TSharedRef<FVariableToAddEntry> NewSubEntry = MakeShared<FVariableToAddEntry>(FAnimNextParamType::FromProperty(InProperty), InProperty->GetFName(), SharedThis(this));
+		NewSubEntry->Parent = NewEntry;
+		NewSubEntry->bIsNew = false;
+		NewEntry->Children.Add(NewSubEntry);
+	};
+
+	auto AddSharedVariablesInternal = [this, &AddVariable](const UScriptStruct* InStruct, auto& InAddSharedVariablesInternal) -> void
+	{
+		// Add the child entries
+		for (TFieldIterator<FProperty> It(InStruct); It; ++It)
+		{
+			AddVariable(*It);
+		}
+	};
+
+	// Add the selected shared variables and its recursive dependents
+	AddSharedVariablesInternal(InStruct, AddSharedVariablesInternal);
 
 	if(NewEntry->Children.Num() > 0)
 	{
@@ -544,14 +631,25 @@ class SVariableToAdd : public SMultiColumnTableRow<TSharedRef<SAddVariablesDialo
 					WeakDialog.Pin()->RefreshCanCreate();
 				});
 			}
-			else if(Entry->EntryType == SAddVariablesDialog::EEntryType::DataInterface)
+			else if(Entry->EntryType == SAddVariablesDialog::EEntryType::SharedVariables)
 			{
 				SAssignNew(EntryWidget, STextBlock)
 				.Font(IDetailLayoutBuilder::GetDetailFontBold())
 				.Text_Lambda([this]()
 				{
-					TSharedPtr<SAddVariablesDialog::FDataInterfaceToAddEntry> InterfaceEntry = StaticCastSharedPtr<SAddVariablesDialog::FDataInterfaceToAddEntry>(Entry);
-					FName Name = InterfaceEntry->DataInterface ? InterfaceEntry->DataInterface->GetFName() : NAME_None; 
+					TSharedPtr<SAddVariablesDialog::FSharedVariablesToAddEntry> SharedVariablesEntry = StaticCastSharedPtr<SAddVariablesDialog::FSharedVariablesToAddEntry>(Entry);
+					FName Name;
+					if (SharedVariablesEntry->SharedVariables.IsType<UAnimNextSharedVariables*>())
+					{
+						UAnimNextSharedVariables* SharedVariables = SharedVariablesEntry->SharedVariables.Get<UAnimNextSharedVariables*>();
+						Name = SharedVariables ? SharedVariables->GetFName() : NAME_None; 
+					}
+					else
+					{
+						const UScriptStruct* SharedVariables = SharedVariablesEntry->SharedVariables.Get<const UScriptStruct*>();
+						Name = SharedVariables ? SharedVariables->GetFName() : NAME_None; 
+					}
+
 					return FText::FromName(Name);
 				});
 				
@@ -727,7 +825,7 @@ void SAddVariablesDialog::HandleGetChildren(TSharedRef<FEntry> InEntry, TArray<T
 	OutChildren = InEntry->Children;
 }
 
-bool SAddVariablesDialog::ShowModal(TArray<FVariableToAdd>& OutVariables, TArray<FDataInterfaceToAdd>& OutDataInterfaces)
+bool SAddVariablesDialog::ShowModal(TArray<FVariableToAdd>& OutVariables, TArray<FSharedVariablesToAdd>& OutSharedVariables)
 {
 	FSlateApplication::Get().AddModalWindow(SharedThis(this), FGlobalTabmanager::Get()->GetRootWindow());
 
@@ -744,8 +842,8 @@ bool SAddVariablesDialog::ShowModal(TArray<FVariableToAdd>& OutVariables, TArray
 				case EEntryType::Variable:
 					OutVariables.Add(*StaticCastSharedRef<FVariableToAddEntry>(Entry));
 					break;
-				case EEntryType::DataInterface:
-					OutDataInterfaces.Add(*StaticCastSharedRef<FDataInterfaceToAddEntry>(Entry));
+				case EEntryType::SharedVariables:
+					OutSharedVariables.Add(*StaticCastSharedRef<FSharedVariablesToAddEntry>(Entry));
 					break;
 				}
 				bHasValid = true;
diff --git a/Engine/Plugins/Experimental/AnimNext/Source/AnimNextEditor/Private/Variables/SAddVariablesDialog.h b/Engine/Plugins/Experimental/AnimNext/Source/AnimNextEditor/Private/Variables/SAddVariablesDialog.h
--- a/Engine/Plugins/Experimental/AnimNext/Source/AnimNextEditor/Private/Variables/SAddVariablesDialog.h
+++ b/Engine/Plugins/Experimental/AnimNext/Source/AnimNextEditor/Private/Variables/SAddVariablesDialog.h
@@ -6,7 +6,7 @@
 #include "Widgets/SWindow.h"
 #include "Widgets/Views/STreeView.h"
 
-class UAnimNextDataInterface;
+class UAnimNextSharedVariables;
 class SWrapBox;
 class UAnimNextModule_EditorData;
 class UAnimNextRigVMAssetEditorData;
@@ -14,13 +14,19 @@ class UAnimNextRigVMAssetEditorData;
 namespace UE::AnimNext::Editor
 {
 
-struct FDataInterfaceToAdd
+struct FSharedVariablesToAdd
 {
-	FDataInterfaceToAdd(UAnimNextDataInterface* InDataInterface)
-		: DataInterface(InDataInterface)
-	{}
+	FSharedVariablesToAdd(UAnimNextSharedVariables* InSharedVariables)
+	{
+		SharedVariables.Set<UAnimNextSharedVariables*>(InSharedVariables);
+	}
+
+	FSharedVariablesToAdd(const UScriptStruct* InSharedVariables)
+	{
+		SharedVariables.Set<const UScriptStruct*>(InSharedVariables);
+	}
 
-	UAnimNextDataInterface* DataInterface;
+	TVariant<UAnimNextSharedVariables*, const UScriptStruct*> SharedVariables;
 };
 
 struct FVariableToAdd
@@ -70,12 +76,12 @@ public:
 
 	void Construct(const FArguments& InArgs, const TArray<UAnimNextRigVMAssetEditorData*>& InAssetEditorDatas);
 
-	bool ShowModal(TArray<FVariableToAdd>& OutVariables, TArray<FDataInterfaceToAdd>& OutDataInterfaces);
+	bool ShowModal(TArray<FVariableToAdd>& OutVariables, TArray<FSharedVariablesToAdd>& OutSharedVariables);
 
 	enum class EEntryType : uint8
 	{
 		Variable,
-		DataInterface,
+		SharedVariables,
 	};
 
 	struct FEntry
@@ -99,7 +105,9 @@ public:
 private:
 	virtual FReply OnKeyDown(const FGeometry& MyGeometry, const FKeyEvent& InKeyEvent) override;
 
-	void AddDataInterface(UAnimNextDataInterface* InDataInterface);
+	void AddSharedVariables(UAnimNextSharedVariables* InSharedVariables);
+
+	void AddSharedVariablesStruct(const UScriptStruct* InStruct);
 
 	void AddEntry(const FAnimNextParamType& InParamType = FAnimNextParamType());
 
@@ -109,11 +117,17 @@ private:
 
 	void DeleteSelectedItems();
 
-	struct FDataInterfaceToAddEntry : FEntry, FDataInterfaceToAdd
+	struct FSharedVariablesToAddEntry : FEntry, FSharedVariablesToAdd
 	{
-		FDataInterfaceToAddEntry(UAnimNextDataInterface* InDataInterface, const TSharedRef<SAddVariablesDialog>& InDialog)
-			: FEntry(EEntryType::DataInterface, InDialog)
-			, FDataInterfaceToAdd(InDataInterface)
+		FSharedVariablesToAddEntry(UAnimNextSharedVariables* InSharedVariables, const TSharedRef<SAddVariablesDialog>& InDialog)
+			: FEntry(EEntryType::SharedVariables, InDialog)
+			, FSharedVariablesToAdd(InSharedVariables)
+		{
+		}
+
+		FSharedVariablesToAddEntry(const UScriptStruct* InSharedVariables, const TSharedRef<SAddVariablesDialog>& InDialog)
+			: FEntry(EEntryType::SharedVariables, InDialog)
+			, FSharedVariablesToAdd(InSharedVariables)
 		{
 		}
 
diff --git a/Engine/Plugins/Experimental/AnimNext/Source/AnimNextEditor/Private/Variables/SVariablesView.cpp b/Engine/Plugins/Experimental/AnimNext/Source/AnimNextEditor/Private/Variables/SVariablesView.cpp
--- a/Engine/Plugins/Experimental/AnimNext/Source/AnimNextEditor/Private/Variables/SVariablesView.cpp
+++ b/Engine/Plugins/Experimental/AnimNext/Source/AnimNextEditor/Private/Variables/SVariablesView.cpp
@@ -129,10 +129,10 @@ FReply SVariablesOutliner::HandleAddVariablesClicked()
 		SNew(SAddVariablesDialog, AssetsToAddTo);
 
 	TArray<FVariableToAdd> VariablesToAdd;
-	TArray<FDataInterfaceToAdd> DataInterfacesToAdd;
-	if(AddVariablesDialog->ShowModal(VariablesToAdd, DataInterfacesToAdd))
+	TArray<FSharedVariablesToAdd> AllSharedVariablesToAdd;
+	if(AddVariablesDialog->ShowModal(VariablesToAdd, AllSharedVariablesToAdd))
 	{
-		FScopedTransaction Transaction(FText::Format(LOCTEXT("AddVariablesFormat", "Add {0}|plural(one=variable, other=variables)"), (DataInterfacesToAdd.Num() + VariablesToAdd.Num()) * AssetsToAddTo.Num()));
+		FScopedTransaction Transaction(FText::Format(LOCTEXT("AddVariablesFormat", "Add {0}|plural(one=variable, other=variables)"), (AllSharedVariablesToAdd.Num() + VariablesToAdd.Num()) * AssetsToAddTo.Num()));
 		for(UAnimNextRigVMAssetEditorData* EditorData : AssetsToAddTo)
 		{
 			for (const FVariableToAdd& VariableToAdd : VariablesToAdd)
@@ -140,9 +140,16 @@ FReply SVariablesOutliner::HandleAddVariablesClicked()
 				EditorData->AddVariable(VariableToAdd.Name, VariableToAdd.Type);
 			}
 
-			for (const FDataInterfaceToAdd& DataInterfaceToAdd : DataInterfacesToAdd)
+			for (const FSharedVariablesToAdd& SharedVariablesToAdd : AllSharedVariablesToAdd)
 			{
-				EditorData->AddDataInterface(DataInterfaceToAdd.DataInterface);
+				if (SharedVariablesToAdd.SharedVariables.IsType<UAnimNextSharedVariables*>())
+				{
+					EditorData->AddSharedVariables(SharedVariablesToAdd.SharedVariables.Get<UAnimNextSharedVariables*>());
+				}
+				else
+				{
+					EditorData->AddSharedVariablesStruct(SharedVariablesToAdd.SharedVariables.Get<const UScriptStruct*>());
+				}
 			}
 		}
 	}
diff --git a/Engine/Plugins/Experimental/AnimNext/Source/AnimNextEditor/Private/Variables/VariableCustomization.cpp b/Engine/Plugins/Experimental/AnimNext/Source/AnimNextEditor/Private/Variables/VariableCustomization.cpp
--- a/Engine/Plugins/Experimental/AnimNext/Source/AnimNextEditor/Private/Variables/VariableCustomization.cpp
+++ b/Engine/Plugins/Experimental/AnimNext/Source/AnimNextEditor/Private/Variables/VariableCustomization.cpp
@@ -39,8 +39,8 @@ void FVariableCustomization::CustomizeDetails(IDetailLayoutBuilder& DetailBuilde
 
 		UAnimNextRigVMAsset* Asset = Variable->GetTypedOuter<UAnimNextRigVMAsset>();
 
-		// Disable access specifier switching specifically for data interfaces
-		if(UAnimNextDataInterface* DataInterface = ExactCast<UAnimNextDataInterface>(Asset))
+		// Disable access specifier switching specifically for shared variables
+		if(UAnimNextSharedVariables* DataInterface = ExactCast<UAnimNextSharedVariables>(Asset))
 		{
 			TSharedPtr<IPropertyHandle> AccessProperty = DetailBuilder.GetProperty(GET_MEMBER_NAME_CHECKED(UAnimNextVariableEntry, Access));
 			AccessProperty->MarkHiddenByCustomization();
diff --git a/Engine/Plugins/Experimental/AnimNext/Source/AnimNextTestSuite/Private/AnimNextVariablesTest.cpp b/Engine/Plugins/Experimental/AnimNext/Source/AnimNextTestSuite/Private/AnimNextVariablesTest.cpp
--- a/Engine/Plugins/Experimental/AnimNext/Source/AnimNextTestSuite/Private/AnimNextVariablesTest.cpp
+++ b/Engine/Plugins/Experimental/AnimNext/Source/AnimNextTestSuite/Private/AnimNextVariablesTest.cpp
@@ -5,6 +5,7 @@
 #include "AnimNextTest.h"
 #include "ScopedTransaction.h"
 #include "UncookedOnlyUtils.h"
+#include "Algo/Compare.h"
 #include "Misc/AutomationTest.h"
 #include "Param/ParamType.h"
 #include "Animation/AnimSequence.h"
@@ -142,8 +143,10 @@ bool FVariableTypesTest::RunTest(const FString& InParameters)
 	AddErrorIfFalse(FAnimNextParamType::GetType<TObjectPtr<UObject>>().IsValid(), TEXT("TObjectPtr<UObject> parameter is invalid."));
 	AddErrorIfFalse(FAnimNextParamType::GetType<UClass*>().IsValid(), TEXT("UClass parameter is invalid."));
 	AddErrorIfFalse(FAnimNextParamType::GetType<TSubclassOf<UObject>>().IsValid(), TEXT("TSubclassOf<UObject> parameter is invalid."));
-	AddErrorIfFalse(FAnimNextParamType::GetType<TSoftObjectPtr<UObject>>().IsValid(), TEXT("TSoftObjectPtr<UObject> parameter is invalid."));
-	AddErrorIfFalse(FAnimNextParamType::GetType<TSoftClassPtr<UObject>>().IsValid(), TEXT("TSoftClassPtr<UObject> parameter is invalid."));
+
+	// TODO: Soft obj/class ptrs not supported yet in RigVM
+//	AddErrorIfFalse(FAnimNextParamType::GetType<TSoftObjectPtr<UObject>>().IsValid(), TEXT("TSoftObjectPtr<UObject> parameter is invalid."));
+//	AddErrorIfFalse(FAnimNextParamType::GetType<TSoftClassPtr<UObject>>().IsValid(), TEXT("TSoftClassPtr<UObject> parameter is invalid."));
 
 	AddErrorIfFalse(FAnimNextParamType::GetType<TArray<bool>>().IsValid(), TEXT("bool array parameter is invalid."));
 	AddErrorIfFalse(FAnimNextParamType::GetType<TArray<uint8>>().IsValid(), TEXT("uint8 array parameter is invalid."));
@@ -163,8 +166,10 @@ bool FVariableTypesTest::RunTest(const FString& InParameters)
 	AddErrorIfFalse(FAnimNextParamType::GetType<TArray<TObjectPtr<UObject>>>().IsValid(), TEXT("TObjectPtr<UObject> array parameter is invalid."));
 	AddErrorIfFalse(FAnimNextParamType::GetType<TArray<UClass*>>().IsValid(), TEXT("UClass array parameter is invalid."));
 	AddErrorIfFalse(FAnimNextParamType::GetType<TArray<TSubclassOf<UObject>>>().IsValid(), TEXT("TSubclassOf<UObject> array parameter is invalid."));
-	AddErrorIfFalse(FAnimNextParamType::GetType<TArray<TSoftObjectPtr<UObject>>>().IsValid(), TEXT("TSoftObjectPtr<UObject> array parameter is invalid."));
-	AddErrorIfFalse(FAnimNextParamType::GetType<TArray<TSoftClassPtr<UObject>>>().IsValid(), TEXT("TSoftClassPtr<UObject> array parameter is invalid."));
+
+	// TODO: Soft obj/class ptrs not supported yet in RigVM
+//	AddErrorIfFalse(FAnimNextParamType::GetType<TArray<TSoftObjectPtr<UObject>>>().IsValid(), TEXT("TSoftObjectPtr<UObject> array parameter is invalid."));
+//	AddErrorIfFalse(FAnimNextParamType::GetType<TArray<TSoftClassPtr<UObject>>>().IsValid(), TEXT("TSoftClassPtr<UObject> array parameter is invalid."));
 
 	// Check type from property
 	#define TEST_ANIMNEXT_PROPERTY(Type, Property) AddErrorIfFalse(FAnimNextParamType::FromProperty(FAnimNextParamTypeTestStruct::StaticStruct()->FindPropertyByName(#Property)) == FAnimNextParamType::GetType<Type>(), #Type " param type is invalid")
@@ -186,8 +191,9 @@ bool FVariableTypesTest::RunTest(const FString& InParameters)
 	TEST_ANIMNEXT_PROPERTY(TObjectPtr<UObject>, Object);
 	TEST_ANIMNEXT_PROPERTY(TObjectPtr<UClass>, Class);
 	TEST_ANIMNEXT_PROPERTY(TSubclassOf<UObject>, SubclassOf);
-	TEST_ANIMNEXT_PROPERTY(TSoftObjectPtr<UObject>, SoftObjectPtr);
-	TEST_ANIMNEXT_PROPERTY(TSoftClassPtr<UObject>, SoftClassPtr);
+	// TODO: Soft obj/class ptrs not supported yet in RigVM
+//	TEST_ANIMNEXT_PROPERTY(TSoftObjectPtr<UObject>, SoftObjectPtr);
+//	TEST_ANIMNEXT_PROPERTY(TSoftClassPtr<UObject>, SoftClassPtr);
 
 	TEST_ANIMNEXT_PROPERTY_ARRAY(bool, Bool);
 	TEST_ANIMNEXT_PROPERTY_ARRAY(uint8, Uint8);
@@ -205,8 +211,9 @@ bool FVariableTypesTest::RunTest(const FString& InParameters)
 	TEST_ANIMNEXT_PROPERTY_ARRAY(TObjectPtr<UObject>, Object);
 	TEST_ANIMNEXT_PROPERTY_ARRAY(TObjectPtr<UClass>, Class);
 	TEST_ANIMNEXT_PROPERTY_ARRAY(TSubclassOf<UObject>, SubclassOf);
-	TEST_ANIMNEXT_PROPERTY_ARRAY(TSoftObjectPtr<UObject>, SoftObjectPtr);
-	TEST_ANIMNEXT_PROPERTY_ARRAY(TSoftClassPtr<UObject>, SoftClassPtr);
+	// TODO: Soft obj/class ptrs not supported yet in RigVM
+//	TEST_ANIMNEXT_PROPERTY_ARRAY(TSoftObjectPtr<UObject>, SoftObjectPtr);
+//	TEST_ANIMNEXT_PROPERTY_ARRAY(TSoftClassPtr<UObject>, SoftClassPtr);
 
 	#undef TEST_ANIMNEXT_PROPERTY
 	#undef TEST_ANIMNEXT_PROPERTY_ARRAY
@@ -405,7 +412,7 @@ bool FVariables_UOLBindings::RunTest(const FString& InParameters)
 	Asset->GetRigVMExtendedExecuteContext().SetRuntimeSettings(RuntimeSettings);
 	
 	FAnimNextExecuteContext& AnimNextContext = Asset->GetRigVMExtendedExecuteContext().GetPublicDataSafe<FAnimNextExecuteContext>();
-	FAnimNextModuleInstance Instance(Asset, NewObject<UAnimNextComponent>(), nullptr, nullptr, EAnimNextModuleInitMethod::None);
+	FAnimNextModuleInstance Instance(Asset, NewObject<UAnimNextComponent>(), nullptr, EAnimNextModuleInitMethod::None);
 	FAnimNextModuleContextData ContextData(&Instance);
 	FScopedExecuteContextData ContextDataScope(AnimNextContext, ContextData);
 
@@ -423,6 +430,163 @@ bool FVariables_UOLBindings::RunTest(const FString& InParameters)
 	return true;
 }
 
+IMPLEMENT_SIMPLE_AUTOMATION_TEST(FVariables_Shared, "Animation.AnimNext.Variables.Shared", EAutomationTestFlags::EditorContext | EAutomationTestFlags::EngineFilter)
+
+namespace Private
+{
+	struct FTestReadWrite
+	{
+		template<typename Type>
+		static bool Test(FAutomationTestBase& Test, FAnimNextModuleInstance& InInstance, const TCHAR* InMemberName, const Type& InValue)
+		{
+			const FAnimNextVariableReference Reference(FName(InMemberName), FAnimNextParamTypeTestStruct::StaticStruct());
+			Test.AddErrorIfFalse(Reference.IsValid(), *FString::Printf(TEXT("FVariables_Shared -> Variable reference to %s was invalid."), InMemberName));
+			EPropertyBagResult SetResult = InInstance.SetVariable<Type>(Reference, InValue);
+			Test.AddErrorIfFalse(SetResult == EPropertyBagResult::Success, *FString::Printf(TEXT("FVariables_Shared -> SetVariable call for %s failed."), InMemberName));
+			Type TestValue;
+			EPropertyBagResult GetResult = InInstance.GetVariable<Type>(Reference, TestValue);
+			Test.AddErrorIfFalse(GetResult == EPropertyBagResult::Success, *FString::Printf(TEXT("FVariables_Shared -> GetVariable call for %s failed."), InMemberName));
+
+			// Specialize comparisons for types that dont have operator==
+			if constexpr(std::is_same_v<Type, FText>)
+			{
+				Test.AddErrorIfFalse(TestValue.EqualTo(InValue), *FString::Printf(TEXT("FVariables_Shared -> Variable value did not match for %s."), InMemberName));
+			}
+			else if constexpr(std::is_same_v<Type, TArray<FText>>)
+			{
+				Test.AddErrorIfFalse(Algo::Compare(TestValue, InValue, [](const FText& A, const FText& B){ return A.EqualTo(B); }), *FString::Printf(TEXT("FVariables_Shared -> Variable value did not match for %s."), InMemberName));
+			}
+			else if constexpr(std::is_same_v<Type, FTransform>)
+			{
+				Test.AddErrorIfFalse(TestValue.Equals(InValue), *FString::Printf(TEXT("FVariables_Shared -> Variable value did not match for %s."), InMemberName));
+			}
+			else if constexpr(std::is_same_v<Type, TArray<FTransform>>)
+			{
+				Test.AddErrorIfFalse(Algo::Compare(TestValue, InValue, [](const FTransform& A, const FTransform& B){ return A.Equals(B); }), *FString::Printf(TEXT("FVariables_Shared -> Variable value did not match for %s."), InMemberName));
+			}
+			else
+			{
+				Test.AddErrorIfFalse(TestValue == InValue, *FString::Printf(TEXT("FVariables_Shared -> Variable value did not match for %s."), InMemberName));
+			}
+			return true;
+		}
+	};
+}
+
+bool FVariables_Shared::RunTest(const FString& InParameters)
+{
+	ON_SCOPE_EXIT{ FUtils::CleanupAfterTests(); };
+
+	UFactory* Factory = NewObject<UFactory>(GetTransientPackage(), UAnimNextModuleFactory::StaticClass());
+	UAnimNextModule* Asset = CastChecked<UAnimNextModule>(Factory->FactoryCreateNew(UAnimNextModule::StaticClass(), GetTransientPackage(), TEXT("TestAsset"), RF_Transient, nullptr, nullptr, NAME_None));
+	UE_RETURN_ON_ERROR(Asset != nullptr, "FVariables_Shared -> Failed to create asset");
+
+	UAnimNextRigVMAssetEditorData* EditorData = UncookedOnly::FUtils::GetEditorData<UAnimNextRigVMAssetEditorData>(Asset);
+	UE_RETURN_ON_ERROR(EditorData != nullptr, "FVariables_Shared -> Asset has no editor data.");
+
+	UAnimNextSharedVariablesEntry* Entry = EditorData->AddSharedVariablesStruct(FAnimNextParamTypeTestStruct::StaticStruct());
+	UE_RETURN_ON_ERROR(Entry != nullptr, "FVariables_Shared -> Could not add shared variables struct.");
+
+	FAnimNextModuleInstance Instance(Asset, NewObject<UAnimNextComponent>(), nullptr, EAnimNextModuleInitMethod::None);
+	Instance.Initialize();
+
+	Private::FTestReadWrite::Test<bool>(*this, Instance, TEXT("bBool"), true);
+	Private::FTestReadWrite::Test<bool>(*this, Instance, TEXT("bBitfield0"), true);
+	Private::FTestReadWrite::Test<bool>(*this, Instance, TEXT("bBitfield1"), true);
+	Private::FTestReadWrite::Test<uint8>(*this, Instance, TEXT("Uint8"), 78);
+	Private::FTestReadWrite::Test<int32>(*this, Instance, TEXT("Int32"), 4578990);
+	Private::FTestReadWrite::Test<int64>(*this, Instance, TEXT("Int64"), 5235987347);
+	Private::FTestReadWrite::Test<float>(*this, Instance, TEXT("Float"), 0.25f);
+	Private::FTestReadWrite::Test<double>(*this, Instance, TEXT("Double"), 8756498576.0);
+	Private::FTestReadWrite::Test<FName>(*this, Instance, TEXT("Name"), FName("TestName"));
+	Private::FTestReadWrite::Test<FString>(*this, Instance, TEXT("String"), TEXT("TestString"));
+	Private::FTestReadWrite::Test<FText>(*this, Instance, TEXT("Text"), NSLOCTEXT("AnimNextTests", "TestText", "TestText"));
+	Private::FTestReadWrite::Test<EPropertyBagContainerType>(*this, Instance, TEXT("Enum"), EPropertyBagContainerType::Set);
+	Private::FTestReadWrite::Test<FAnimNextParamType>(*this, Instance, TEXT("Struct"), FAnimNextParamType::GetType<FAnimNextParamType>());
+	Private::FTestReadWrite::Test<FVector>(*this, Instance, TEXT("Vector"), FVector(1.0f, 2.0f, 3.0f));
+	Private::FTestReadWrite::Test<FTransform>(*this, Instance, TEXT("Transform"), FTransform(FQuat::MakeFromRotator(FRotator(1.0f, 2.0f, 3.0f)), FVector(3.0f, 2.0f, 1.0f), FVector(1.0f, 2.0f, 3.0f)));
+	Private::FTestReadWrite::Test<TObjectPtr<UObject>>(*this, Instance, TEXT("Object"), UAnimNextRigVMAsset::StaticClass()->GetDefaultObject());
+	Private::FTestReadWrite::Test<TObjectPtr<UClass>>(*this, Instance, TEXT("Class"), UAnimNextRigVMAsset::StaticClass());
+	Private::FTestReadWrite::Test<TSubclassOf<UObject>>(*this, Instance, TEXT("SubclassOf"), UAnimNextRigVMAsset::StaticClass());
+	// TODO: Soft obj/class ptrs not supported yet in RigVM
+//	Private::FTestReadWrite::Test<TSoftObjectPtr<UObject>>(*this, Instance, TEXT("SoftObjectPtr"), UAnimNextRigVMAsset::StaticClass()->GetDefaultObject());
+//	Private::FTestReadWrite::Test<TSoftClassPtr<UObject>>(*this, Instance, TEXT("SoftClassPtr"), UAnimNextRigVMAsset::StaticClass());
+
+	Private::FTestReadWrite::Test<TArray<bool>>(*this, Instance, TEXT("BoolArray"), { true });
+	Private::FTestReadWrite::Test<TArray<uint8>>(*this, Instance, TEXT("Uint8Array"), { 78 });
+	Private::FTestReadWrite::Test<TArray<int32>>(*this, Instance, TEXT("Int32Array"), { 4578990 });
+	Private::FTestReadWrite::Test<TArray<int64>>(*this, Instance, TEXT("Int64Array"), { 5235987347 });
+	Private::FTestReadWrite::Test<TArray<float>>(*this, Instance, TEXT("FloatArray"), { 0.25f });
+	Private::FTestReadWrite::Test<TArray<double>>(*this, Instance, TEXT("DoubleArray"), { 8756498576.0 });
+	Private::FTestReadWrite::Test<TArray<FName>>(*this, Instance, TEXT("NameArray"), { FName("TestName") });
+	Private::FTestReadWrite::Test<TArray<FString>>(*this, Instance, TEXT("StringArray"), { TEXT("TestString") });
+	Private::FTestReadWrite::Test<TArray<FText>>(*this, Instance, TEXT("TextArray"), { NSLOCTEXT("AnimNextTests", "TestText", "TestText") });
+	Private::FTestReadWrite::Test<TArray<EPropertyBagContainerType>>(*this, Instance, TEXT("EnumArray"), { EPropertyBagContainerType::Set });
+	Private::FTestReadWrite::Test<TArray<FAnimNextParamType>>(*this, Instance, TEXT("StructArray"), { FAnimNextParamType::GetType<FAnimNextParamType>() });
+	Private::FTestReadWrite::Test<TArray<FVector>>(*this, Instance, TEXT("VectorArray"), { FVector(1.0f, 2.0f, 3.0f) });
+	Private::FTestReadWrite::Test<TArray<FTransform>>(*this, Instance, TEXT("TransformArray"), { FTransform(FQuat::MakeFromRotator(FRotator(1.0f, 2.0f, 3.0f)), FVector(3.0f, 2.0f, 1.0f), FVector(1.0f, 2.0f, 3.0f)) });
+	Private::FTestReadWrite::Test<TArray<TObjectPtr<UObject>>>(*this, Instance, TEXT("ObjectArray"), { UAnimNextRigVMAsset::StaticClass()->GetDefaultObject() });
+	Private::FTestReadWrite::Test<TArray<TObjectPtr<UClass>>>(*this, Instance, TEXT("ClassArray"), { UAnimNextRigVMAsset::StaticClass() });
+	Private::FTestReadWrite::Test<TArray<TSubclassOf<UObject>>>(*this, Instance, TEXT("SubclassOfArray"), { UAnimNextRigVMAsset::StaticClass() });
+	// TODO: Soft obj/class ptrs not supported yet in RigVM
+//	Private::FTestReadWrite::Test<TArray<TSoftObjectPtr<UObject>>>(*this, Instance, TEXT("SoftObjectPtrArray"), { UAnimNextRigVMAsset::StaticClass()->GetDefaultObject() });
+//	Private::FTestReadWrite::Test<TArray<TSoftClassPtr<UObject>>>(*this, Instance, TEXT("SoftClassPtrArray"), { UAnimNextRigVMAsset::StaticClass() });
+
+	bool bResult = Instance.AccessVariablesStruct<FAnimNextParamTypeTestStruct>([this](FAnimNextParamTypeTestStruct& InStruct)
+	{
+		auto TestValue = [this]<typename T0>(const TCHAR* InMemberName, const auto& InValue, const T0& InComparison)
+		{
+			// Specialize comparisons for types that dont have operator==
+			if constexpr(std::is_same_v<T0, FText>)
+			{
+				AddErrorIfFalse(InComparison.EqualTo(InValue), *FString::Printf(TEXT("FVariables_Shared -> Set value does not match for %s."), InMemberName));
+			}
+			else if constexpr(std::is_same_v<T0, TArray<FText>>)
+			{
+				AddErrorIfFalse(Algo::Compare(InComparison, InValue, [](const FText& A, const FText& B){ return A.EqualTo(B); }), *FString::Printf(TEXT("FVariables_Shared -> Set value does not match for %s."), InMemberName));
+			}
+			else if constexpr(std::is_same_v<T0, FTransform>)
+			{
+				AddErrorIfFalse(InComparison.Equals(InValue), *FString::Printf(TEXT("FVariables_Shared -> Set value does not match for %s."), InMemberName));
+			}
+			else if constexpr(std::is_same_v<T0, TArray<FTransform>>)
+			{
+				AddErrorIfFalse(Algo::Compare(InComparison, InValue, [](const FTransform& A, const FTransform& B){ return A.Equals(B); }), *FString::Printf(TEXT("FVariables_Shared -> Set value does not match for %s."), InMemberName));
+			}
+			else
+			{
+				AddErrorIfFalse(InValue == InComparison, *FString::Printf(TEXT("FVariables_Shared -> Set value does not match for %s."), InMemberName));
+			}
+		};
+
+		TestValue(TEXT("bBool"), InStruct.bBool, true);
+		TestValue(TEXT("bBitfield0"), InStruct.bBitfield0 != 0, true);
+		TestValue(TEXT("bBitfield1"), InStruct.bBitfield1 != 0, true);
+		TestValue(TEXT("Uint8"), InStruct.Uint8, 78);
+		TestValue(TEXT("Int32"), InStruct.Int32, 4578990);
+		TestValue(TEXT("Int64"), InStruct.Int64, 5235987347);
+		TestValue(TEXT("Float"), InStruct.Float, 0.25f);
+		TestValue(TEXT("Double"), InStruct.Double, 8756498576.0);
+		TestValue(TEXT("Name"), InStruct.Name, FName("TestName"));
+		TestValue(TEXT("String"), InStruct.String, FString("TestString"));
+		TestValue(TEXT("Text"), InStruct.Text, NSLOCTEXT("AnimNextTests", "TestText", "TestText"));
+		TestValue(TEXT("Enum"), InStruct.Enum, EPropertyBagContainerType::Set);
+		TestValue(TEXT("Struct"), InStruct.Struct, FAnimNextParamType::GetType<FAnimNextParamType>());
+		TestValue(TEXT("Vector"), InStruct.Vector, FVector(1.0f, 2.0f, 3.0f));
+		TestValue(TEXT("Transform"), InStruct.Transform, FTransform(FQuat::MakeFromRotator(FRotator(1.0f, 2.0f, 3.0f)), FVector(3.0f, 2.0f, 1.0f), FVector(1.0f, 2.0f, 3.0f)));
+		TestValue(TEXT("Object"), InStruct.Object, UAnimNextRigVMAsset::StaticClass()->GetDefaultObject());
+		TestValue(TEXT("Class"), InStruct.Class, UAnimNextRigVMAsset::StaticClass());
+		TestValue(TEXT("SubclassOf"), InStruct.SubclassOf, UAnimNextRigVMAsset::StaticClass());
+		// TODO: Soft obj/class ptrs not supported yet in RigVM
+	//	TestValue(TEXT("SoftObjectPtr"), InStruct.SoftObjectPtr, UAnimNextRigVMAsset::StaticClass()->GetDefaultObject());
+	//	TestValue(TEXT("SoftClassPtr"), InStruct.SoftClassPtr, UAnimNextRigVMAsset::StaticClass());
+	});
+
+	UE_RETURN_ON_ERROR(bResult, "FVariables_Shared -> Could not access variable struct on instance.");
+
+	return true;
+}
+
 }
 
 #endif	// WITH_DEV_AUTOMATION_TESTS
diff --git a/Engine/Plugins/Experimental/AnimNext/Source/AnimNextTestSuite/Private/AnimNextVariablesTest.h b/Engine/Plugins/Experimental/AnimNext/Source/AnimNextTestSuite/Private/AnimNextVariablesTest.h
--- a/Engine/Plugins/Experimental/AnimNext/Source/AnimNextTestSuite/Private/AnimNextVariablesTest.h
+++ b/Engine/Plugins/Experimental/AnimNext/Source/AnimNextTestSuite/Private/AnimNextVariablesTest.h
@@ -11,7 +11,7 @@
 
 class UAnimNextComponent;
 
-USTRUCT()
+USTRUCT(BlueprintInternalUseOnly)
 struct FAnimNextParamTypeTestStruct
 {
 	GENERATED_BODY()
@@ -19,6 +19,12 @@ struct FAnimNextParamTypeTestStruct
 	UPROPERTY()
 	bool bBool = false;
 
+	UPROPERTY()
+	uint8 bBitfield0 : 1 = false;
+
+	UPROPERTY()
+	uint8 bBitfield1 : 1 = false;
+	
 	UPROPERTY()
 	uint8 Uint8 = 0;
 
@@ -67,11 +73,12 @@ struct FAnimNextParamTypeTestStruct
 	UPROPERTY()
 	TSubclassOf<UObject> SubclassOf;
 
-	UPROPERTY()
-	TSoftObjectPtr<UObject> SoftObjectPtr;
+	// TODO: Soft obj/class ptrs not supported yet in RigVM
+//	UPROPERTY()
+//	TSoftObjectPtr<UObject> SoftObjectPtr;
 
-	UPROPERTY()
-	TSoftClassPtr<UObject> SoftClassPtr;
+//	UPROPERTY()
+//	TSoftClassPtr<UObject> SoftClassPtr;
 
 	UPROPERTY()
 	TArray<bool> BoolArray;
@@ -124,11 +131,12 @@ struct FAnimNextParamTypeTestStruct
 	UPROPERTY()
 	TArray<TSubclassOf<UObject>> SubclassOfArray;
 
-	UPROPERTY()
-	TArray<TSoftObjectPtr<UObject>> SoftObjectPtrArray;
+	// TODO: Soft obj/class ptrs not supported yet in RigVM
+//	UPROPERTY()
+//	TArray<TSoftObjectPtr<UObject>> SoftObjectPtrArray;
 
-	UPROPERTY()
-	TArray<TSoftClassPtr<UObject>> SoftClassPtrArray;
+//	UPROPERTY()
+//	TArray<TSoftClassPtr<UObject>> SoftClassPtrArray;
 };
 
 USTRUCT(meta=(Hidden))
diff --git a/Engine/Plugins/Experimental/AnimNext/Source/AnimNextUncookedOnly/AnimNextUncookedOnly.Build.cs b/Engine/Plugins/Experimental/AnimNext/Source/AnimNextUncookedOnly/AnimNextUncookedOnly.Build.cs
--- a/Engine/Plugins/Experimental/AnimNext/Source/AnimNextUncookedOnly/AnimNextUncookedOnly.Build.cs
+++ b/Engine/Plugins/Experimental/AnimNext/Source/AnimNextUncookedOnly/AnimNextUncookedOnly.Build.cs
@@ -37,6 +37,11 @@ namespace UnrealBuildTool.Rules
 				}
 			);
 
+			if(Target.bWithLiveCoding)
+			{
+				PrivateDependencyModuleNames.Add("LiveCoding");
+			}
+			
 			if (Target.bBuildEditor)
 			{
 				PrivateDependencyModuleNames.AddRange(
diff --git a/Engine/Plugins/Experimental/AnimNext/Source/AnimNextUncookedOnly/Internal/AnimNextAssetWorkspaceAssetUserData.h b/Engine/Plugins/Experimental/AnimNext/Source/AnimNextUncookedOnly/Internal/AnimNextAssetWorkspaceAssetUserData.h
--- a/Engine/Plugins/Experimental/AnimNext/Source/AnimNextUncookedOnly/Internal/AnimNextAssetWorkspaceAssetUserData.h
+++ b/Engine/Plugins/Experimental/AnimNext/Source/AnimNextUncookedOnly/Internal/AnimNextAssetWorkspaceAssetUserData.h
@@ -8,7 +8,7 @@
 #include "IAnimNextRigVMGraphInterface.h"
 #include "RigVMModel/Nodes/RigVMFunctionReferenceNode.h"
 #include "Module/AnimNextModule.h"
-#include "DataInterface/AnimNextDataInterface.h"
+#include "Variables/AnimNextSharedVariables.h"
 
 #include "AnimNextAssetWorkspaceAssetUserData.generated.h"
 
@@ -52,9 +52,9 @@ struct FAnimNextDataInterfaceOutlinerData : public FAnimNextRigVMAssetOutlinerDa
 
 	FAnimNextDataInterfaceOutlinerData() = default;
 
-	UAnimNextDataInterface* GetDataInterface() const
+	UAnimNextSharedVariables* GetDataInterface() const
 	{
-		return Cast<UAnimNextDataInterface>(GetAsset());
+		return Cast<UAnimNextSharedVariables>(GetAsset());
 	}
 };
 
diff --git a/Engine/Plugins/Experimental/AnimNext/Source/AnimNextUncookedOnly/Internal/AnimNextControllerBase.h b/Engine/Plugins/Experimental/AnimNext/Source/AnimNextUncookedOnly/Internal/AnimNextControllerBase.h
new file mode 100644
--- /dev/null
+++ b/Engine/Plugins/Experimental/AnimNext/Source/AnimNextUncookedOnly/Internal/AnimNextControllerBase.h
@@ -0,0 +1,38 @@
+﻿// Copyright Epic Games, Inc. All Rights Reserved.
+
+#pragma once
+
+#include "CoreMinimal.h"
+#include "RigVMModel/RigVMController.h"
+#include "AnimNextControllerBase.generated.h"
+
+class UAnimNextSharedVariables;
+class UAnimNextSharedVariableNode;
+class UAnimNextRigVMAssetEditorData;
+
+/**
+  * Implements AnimNext RigVM controller extensions
+  */
+UCLASS(MinimalAPI)
+class UAnimNextControllerBase : public URigVMController
+{
+	GENERATED_BODY()
+
+public:
+	// Add a shared variable from the specified asset
+	UFUNCTION(BlueprintCallable, Category = RigVMController)
+	ANIMNEXTUNCOOKEDONLY_API UAnimNextSharedVariableNode* AddAssetSharedVariableNode(const UAnimNextSharedVariables* InAsset, const FName& InVariableName, const FString& InCPPType, UObject* InCPPTypeObject, bool bIsGetter, const FString& InDefaultValue, const FVector2D& InPosition = FVector2D::ZeroVector, const FString& InNodeName = TEXT(""), bool bSetupUndoRedo = true, bool bPrintPythonCommand = false);
+
+	// Add a shared variable from the specified struct
+	UFUNCTION(BlueprintCallable, Category = RigVMController)
+	ANIMNEXTUNCOOKEDONLY_API UAnimNextSharedVariableNode* AddStructSharedVariableNode(const UScriptStruct* InStruct, const FName& InVariableName, const FString& InCPPType, UObject* InCPPTypeObject, bool bIsGetter, const FString& InDefaultValue, const FVector2D& InPosition = FVector2D::ZeroVector, const FString& InNodeName = TEXT(""), bool bSetupUndoRedo = true, bool bPrintPythonCommand = false);
+
+	// Add a shared variable from the specified source object
+	ANIMNEXTUNCOOKEDONLY_API UAnimNextSharedVariableNode* AddSharedVariableNode(const FString& InSourceObjectPath, const FName& InVariableName, const FString& InCPPType, UObject* InCPPTypeObject, bool bIsGetter, const FString& InDefaultValue, const FVector2D& InPosition, const FString& InNodeName, bool bSetupUndoRedo, bool bPrintPythonCommand);
+
+private:
+	friend UAnimNextRigVMAssetEditorData;
+
+	// Helper function to upgrade older assets
+	void ReplaceVariableNodeWithSharedVariableNode(URigVMVariableNode* InVariableNode, FName InNewVariableName, const UObject* InAssetOrStruct);
+};
