AnimNext: Pass-by-ref runtime refactor

diff --git a/Engine/Plugins/Experimental/AnimNext/Source/AnimNextUncookedOnly/Internal/AnimNextRigVMAssetEditorData.h b/Engine/Plugins/Experimental/AnimNext/Source/AnimNextUncookedOnly/Internal/AnimNextRigVMAssetEditorData.h
--- a/Engine/Plugins/Experimental/AnimNext/Source/AnimNextUncookedOnly/Internal/AnimNextRigVMAssetEditorData.h
+++ b/Engine/Plugins/Experimental/AnimNext/Source/AnimNextUncookedOnly/Internal/AnimNextRigVMAssetEditorData.h
@@ -3,6 +3,8 @@
 #pragma once
 
 #include "CoreMinimal.h"
+#include "AnimNextControllerBase.h"
+#include "IAnimNextRigVMExportInterface.h"
 #include "RigVMModel/RigVMGraph.h"
 #include "RigVMCore/RigVMGraphFunctionHost.h"
 #include "RigVMBlueprint.h"
@@ -10,6 +12,7 @@
 #include "Kismet/BlueprintFunctionLibrary.h"
 #include "Entries/AnimNextRigVMAssetEntry.h"
 #include "Variables/AnimNextProgrammaticVariable.h"
+#include "Variables/AnimNextVariableReference.h"
 #include "AnimNextRigVMAssetEditorData.generated.h"
 
 #define UE_API ANIMNEXTUNCOOKEDONLY_API
@@ -17,7 +20,7 @@
 enum class ERigVMGraphNotifType : uint8;
 class UAnimNextRigVMAssetEditorData;
 class UAnimNextEdGraph;
-class UAnimNextDataInterfaceEntry;
+class UAnimNextSharedVariablesEntry;
 class UAnimNextWorkspaceEditorMode;
 struct FAnimNextRigVMAssetCompileContext;
 struct FAnimNextGetFunctionHeaderCompileContext;
@@ -123,10 +126,14 @@ public:
 	UFUNCTION(BlueprintCallable, Category = "AnimNext|Entries", meta=(ScriptMethod))
 	static ANIMNEXTUNCOOKEDONLY_API UAnimNextEventGraphEntry* AddEventGraph(UAnimNextRigVMAsset* InAsset, FName InName, UScriptStruct* InEventStruct, bool bSetupUndoRedo = true, bool bPrintPythonCommand = true);
 
-	/** Adds a data interface to an AnimNext asset */
+	/** Adds shared variables to an AnimNext asset */
 	UFUNCTION(BlueprintCallable, Category = "UAF|Entries", meta=(ScriptMethod))
-	static ANIMNEXTUNCOOKEDONLY_API UAnimNextDataInterfaceEntry* AddDataInterface(UAnimNextRigVMAsset* InAsset, UAnimNextDataInterface* InDataInterface, bool bSetupUndoRedo = true, bool bPrintPythonCommand = true);
+	static ANIMNEXTUNCOOKEDONLY_API UAnimNextSharedVariablesEntry* AddSharedVariables(UAnimNextRigVMAsset* InAsset, UAnimNextSharedVariables* InSharedVariables, bool bSetupUndoRedo = true, bool bPrintPythonCommand = true);
 
+	/** Adds shared variables struct to an AnimNext asset */
+	UFUNCTION(BlueprintCallable, Category = "AnimNext|Entries", meta=(ScriptMethod))
+	static ANIMNEXTUNCOOKEDONLY_API UAnimNextSharedVariablesEntry* AddSharedVariablesStruct(UAnimNextRigVMAsset* InAsset, UScriptStruct* InStruct, bool bSetupUndoRedo = true, bool bPrintPythonCommand = true);
+	
 	/** Adds a function to an AnimNext asset */
 	UFUNCTION(BlueprintCallable, Category = "UAF|Entries", meta=(ScriptMethod))
 	static ANIMNEXTUNCOOKEDONLY_API URigVMLibraryNode* AddFunction(UAnimNextRigVMAsset* InAsset, FName InFunctionName, bool bInMutable, bool bSetupUndoRedo = true, bool bPrintPythonCommand = true);
@@ -145,8 +152,11 @@ public:
 	/** Adds an event graph to this asset */
 	UE_API UAnimNextEventGraphEntry* AddEventGraph(FName InName, UScriptStruct* InEventStruct, bool bSetupUndoRedo = true, bool bPrintPythonCommand = true);
 
-	/** Adds a data interface to this asset */
-	UE_API UAnimNextDataInterfaceEntry* AddDataInterface(UAnimNextDataInterface* InDataInterface, bool bSetupUndoRedo = true, bool bPrintPythonCommand = true);
+	/** Adds shared variables to this asset */
+	UE_API UAnimNextSharedVariablesEntry* AddSharedVariables(const UAnimNextSharedVariables* InSharedVariables, bool bSetupUndoRedo = true, bool bPrintPythonCommand = true);
+
+	/** Adds shared variables struct to this asset */
+	UE_API UAnimNextSharedVariablesEntry* AddSharedVariablesStruct(const UScriptStruct* InStruct, bool bSetupUndoRedo = true, bool bPrintPythonCommand = true);
 
 	/** Adds a function to this asset */
 	UE_API URigVMLibraryNode* AddFunction(FName InFunctionName, bool bInMutable, bool bInSetupUndoRedo = true, bool bPrintPythonCommand = true);
@@ -199,7 +209,7 @@ protected:
 	friend class UAnimNextModuleWorkspaceAssetUserData;
 	friend class UE::AnimNext::Editor::SAddVariablesDialog;
 	friend class UE::AnimNext::Editor::FVariableProxyCustomization;
-	friend class UAnimNextDataInterfaceEntry;
+	friend class UAnimNextSharedVariablesEntry;
 	friend struct UE::AnimNext::UncookedOnly::FPublicVariablesImpl;
 	friend UE::AnimNext::Editor::FAnimNextAnimGraphEditorModule;
 	friend class UE::AnimNext::Editor::FCallFunctionSharedDataDetails;
@@ -286,7 +296,7 @@ protected:
 	UE_API virtual TObjectPtr<URigVMGraph> CreateContainedGraphModel(URigVMCollapseNode* CollapseNode, const FName& Name) override;
 
 	// Override called during initialization to determine what RigVM controller class is used
-	virtual TSubclassOf<URigVMController> GetControllerClass() const { return URigVMController::StaticClass(); }
+	virtual TSubclassOf<URigVMController> GetControllerClass() const { return UAnimNextControllerBase::StaticClass(); }
 
 	// Override called during initialization to determine what RigVM execute struct is used
 	virtual UScriptStruct* GetExecuteContextStruct() const PURE_VIRTUAL(UAnimNextRigVMAssetEditorData::GetExecuteContextStruct, return nullptr;)
@@ -471,13 +481,35 @@ protected:
 	// Checks all entries to see if any are public variables
 	UE_API bool HasPublicVariables() const;
 
-	// Gets any public variables that this asset has. Variables have no specified order.
-	// May recurse into implemented data interfaces, so public variables returned by this function may not be directly owned by this asset.
-	UE_API void GetPublicVariables(TArray<UAnimNextVariableEntry*>& OutPublicVariables) const;
+	// Information about a variable returned from GetAllVariables
+	struct FVariableInfo
+	{
+		FName Name;
+		FAnimNextParamType Type;
+		const UObject* SourceObject = nullptr;
+		EAnimNextExportAccessSpecifier Access = EAnimNextExportAccessSpecifier::Private;
+		const FProperty* Property = nullptr;
+		TConstArrayView<uint8> DefaultValue;
+	};
+
+	// How to recurse when getting variable info
+	enum class EVariableRecursion
+	{
+		SelfOnly,
+		IncludeShared,
+	};
 
-	// Gets any variables that this asset has. Variables have no specified order.
-	// May recurse into implemented data interfaces (for public variables only), so variables returned by this function may not be directly owned by this asset.
-	UE_API void GetAllVariables(TArray<UAnimNextVariableEntry*>& OutPublicVariables) const;
+	// How to filter variables by access level when getting variable info
+	enum class EVariableAccessFilter
+	{
+		PublicOnly,
+		All,
+	};
+	
+	// Gets any variables that this asset has according to the pass-in params. Variables have no specified order.
+	// May recurse into shared variables, so variables returned by this function may not be directly owned by this asset. In this case the
+	// SourceObject will indicate where the variable came from.
+	UE_API void GetAllVariables(TArray<FVariableInfo>& OutVariables, EVariableRecursion InRecursion, EVariableAccessFilter InAccess) const;
 
 	// Refresh the 'external' models for the RigVM client to reference
 	UE_API void RefreshExternalModels();
@@ -512,6 +544,19 @@ protected:
 	// Check whether this asset should be recompiled
 	UE_API bool IsDirtyForRecompilation() const;
 
+	// Called during compilation to generate internal variables used for natively-callable RigVM functions
+	void BuildFunctionWrapperEventVariables(FAnimNextRigVMAssetCompileContext& InContext);
+
+	// Called during compilation to generate wrapper events for natively-callable RigVM functions
+	void BuildFunctionWrapperEvents(FAnimNextRigVMAssetCompileContext& InContext, const FRigVMCompileSettings& InSettings);
+
+	// Generates a combined property bag consisting of all variables both internal and shared
+	// NOTE: Uses compiled data from the hosting asset uses and thus can introduce circular dependency issues if misused!
+	UE_API FInstancedPropertyBag GenerateCombinedPropertyBag() const;
+
+	// Serialization/post-load fixup for old data
+	void UpgradeDataInterfaces();
+
 	/** All entries in this asset - not saved, discovered at load time and also contains InternalEntries */
 	UPROPERTY(transient)
 	TArray<TObjectPtr<UAnimNextRigVMAssetEntry>> Entries;
@@ -528,12 +573,13 @@ protected:
 
 	// The native C++ struct that is used to communicate with this asset
 	UE_DEPRECATED(5.6, "This property is deprecated. Please use NativeInterfaces instead")
-	UPROPERTY(meta = (ShowDisplayNames, MetaStruct = "/Script/AnimNext.AnimNextNativeDataInterface"))
+	UPROPERTY()
 	TObjectPtr<const UScriptStruct> NativeInterface_DEPRECATED;
 
 	// The list of native C++ structs that are used to communicate with this asset
-	UPROPERTY(EditAnywhere, Category = "Native", meta = (ShowDisplayNames, MetaStruct = "/Script/AnimNext.AnimNextNativeDataInterface"))
-	TArray<TObjectPtr<const UScriptStruct>> NativeInterfaces;
+	UE_DEPRECATED(5.6, "This property is deprecated. Please use NativeInterfaces instead")
+	UPROPERTY()
+	TArray<TObjectPtr<const UScriptStruct>> NativeInterfaces_DEPRECATED;
 
 	UPROPERTY(EditAnywhere, Category = "User Interface")
 	FRigVMEdGraphDisplaySettings RigGraphDisplaySettings;
@@ -562,6 +608,13 @@ protected:
 	UPROPERTY(VisibleAnywhere, Category = "External Packages", AssetRegistrySearchable)
 	bool bUsesExternalPackages = true;
 
+	UPROPERTY()
+	FName DefaultInjectionSite_DEPRECATED;
+
+	// The injection site (AnimNext Graph) that will be used by default
+	UPROPERTY(EditAnywhere, DisplayName = "Default Injection Site", Category=Injection, meta=(AllowedType = FAnimNextAnimGraph))
+	FAnimNextVariableReference DefaultInjectionSiteReference;
+
 	FOnRigVMCompiledEvent RigVMCompiledEvent;
 
 	FRigVMGraphModifiedEvent RigVMGraphModifiedEvent;
@@ -595,6 +648,8 @@ protected:
 	bool bSuspendPythonMessagesForRigVMClient = true;
 	bool bSuspendEditorDataNotifications = false;
 	bool bSuspendCompilationNotifications = false;
+	bool bSuspendCompilerReports = false;
+	bool bUpgradeDataInterfacesOnLoad = false;
 };
 
 #undef UE_API
diff --git a/Engine/Plugins/Experimental/AnimNext/Source/AnimNextUncookedOnly/Internal/AnimNextSharedVariableNode.h b/Engine/Plugins/Experimental/AnimNext/Source/AnimNextUncookedOnly/Internal/AnimNextSharedVariableNode.h
new file mode 100644
--- /dev/null
+++ b/Engine/Plugins/Experimental/AnimNext/Source/AnimNextUncookedOnly/Internal/AnimNextSharedVariableNode.h
@@ -0,0 +1,31 @@
+﻿// Copyright Epic Games, Inc. All Rights Reserved.
+
+#pragma once
+
+#include "Entries/AnimNextSharedVariablesEntry.h"
+#include "RigVMModel/Nodes/RigVMVariableNode.h"
+#include "AnimNextSharedVariableNode.generated.h"
+
+class UAnimNextSharedVariables;
+class UAnimNextControllerBase;
+
+UCLASS()
+class UAnimNextSharedVariableNode : public URigVMVariableNode
+{
+	GENERATED_BODY()
+
+	// URigVMVariableNode interface
+	virtual FString GetNodeSubTitle() const override;
+
+private:
+	friend UAnimNextControllerBase;
+
+	UPROPERTY()
+	TObjectPtr<const UAnimNextSharedVariables> Asset;
+
+	UPROPERTY()
+	TObjectPtr<const UScriptStruct> Struct;
+	
+	UPROPERTY()
+	EAnimNextSharedVariablesType Type;
+};
\ No newline at end of file
diff --git a/Engine/Plugins/Experimental/AnimNext/Source/AnimNextUncookedOnly/Internal/Compilation/AnimNextProgrammaticFunctionHeader.h b/Engine/Plugins/Experimental/AnimNext/Source/AnimNextUncookedOnly/Internal/Compilation/AnimNextProgrammaticFunctionHeader.h
--- a/Engine/Plugins/Experimental/AnimNext/Source/AnimNextUncookedOnly/Internal/Compilation/AnimNextProgrammaticFunctionHeader.h
+++ b/Engine/Plugins/Experimental/AnimNext/Source/AnimNextUncookedOnly/Internal/Compilation/AnimNextProgrammaticFunctionHeader.h
@@ -13,20 +13,6 @@
 struct FAnimNextProgrammaticFunctionHeader
 {
 	/** RigVM Function Header we are wrapping */
-	FRigVMGraphFunctionHeader Wrapped;
-
-	/** 
-	 * True instructs compiler that we should generate & link programmatic variables for this function's params
-	 * 
-	 * Used by those who cannot set params by linking into the function node during graph compile time (Ex: StateTree).
-	 */
-	bool bGenerateParamVariables = false;
-
-	/** 
-	 * True instructs compiler that we should generate & link programmatic variables for this function's returns  
-	 * 
-	 * Used by those who cannot get returns by linking into the function node during graph compile time (Ex: StateTree).
-	 */
-	bool bGenerateReturnVariables = false;
+	FRigVMGraphFunctionHeader FunctionHeader;
 };
 
diff --git a/Engine/Plugins/Experimental/AnimNext/Source/AnimNextUncookedOnly/Internal/Entries/AnimNextSharedVariablesEntry.h b/Engine/Plugins/Experimental/AnimNext/Source/AnimNextUncookedOnly/Internal/Entries/AnimNextSharedVariablesEntry.h
new file mode 100644
--- /dev/null
+++ b/Engine/Plugins/Experimental/AnimNext/Source/AnimNextUncookedOnly/Internal/Entries/AnimNextSharedVariablesEntry.h
@@ -0,0 +1,112 @@
+﻿// Copyright Epic Games, Inc. All Rights Reserved.
+
+#pragma once
+
+#include "CoreMinimal.h"
+#include "AnimNextRigVMAssetEntry.h"
+#include "Param/ParamType.h"
+#include "StructUtils/PropertyBag.h"
+#include "AnimNextSharedVariablesEntry.generated.h"
+
+class UAnimNextSharedVariables;
+class UAnimNextRigVMAssetEditorData;
+class UAssetDefinition_AnimNextSharedVariablesEntry;
+
+namespace UE::AnimNext::Editor
+{
+	class FVariablesOutlinerHierarchy;
+	struct FVariablesOutlinerDataInterfaceItem;
+	class SVariablesOutlinerDataInterfaceLabel;
+	class SVariablesOutlinerValue;
+	class FVariableProxyCustomization;
+	class SAddVariablesDialog;
+	class SVariableOverride;
+	class FVariablesOutlinerMode;
+}
+
+namespace UE::AnimNext::UncookedOnly
+{
+	struct FUtils;
+}
+
+namespace UE::AnimNext::Tests
+{
+	class FDataInterfaceCompile;
+}
+
+UENUM()
+enum class EAnimNextSharedVariablesType : uint8
+{
+	Asset,
+	Struct
+};
+
+UCLASS(MinimalAPI, Category = "Shared Variables", DisplayName = "Shared Variables")
+class UAnimNextSharedVariablesEntry : public UAnimNextRigVMAssetEntry
+{
+	GENERATED_BODY()
+
+	friend class UAnimNextRigVMAssetEditorData;
+	friend class UAssetDefinition_AnimNextSharedVariablesEntry;
+	friend class UE::AnimNext::Editor::FVariablesOutlinerHierarchy;
+	friend struct UE::AnimNext::Editor::FVariablesOutlinerDataInterfaceItem;
+	friend struct UE::AnimNext::UncookedOnly::FUtils;
+	friend class UE::AnimNext::Tests::FDataInterfaceCompile;
+	friend class UE::AnimNext::Editor::SVariablesOutlinerDataInterfaceLabel;
+	friend class UE::AnimNext::Editor::SVariablesOutlinerValue;
+	friend class UE::AnimNext::Editor::FVariableProxyCustomization;
+	friend class UE::AnimNext::Editor::SAddVariablesDialog;
+	friend class UE::AnimNext::Editor::SVariableOverride;
+	friend UE::AnimNext::Editor::FVariablesOutlinerMode;
+
+	// UAnimNextRigVMAssetEntry interface
+	virtual void Initialize(UAnimNextRigVMAssetEditorData* InEditorData) override;
+	virtual FName GetEntryName() const override;
+	virtual void SetEntryName(FName InName, bool bSetupUndoRedo = true) override {}
+	virtual FText GetDisplayName() const override;
+	virtual FText GetDisplayNameTooltip() const override;
+
+	// Get the type of this shared variables entry
+	EAnimNextSharedVariablesType GetType() const
+	{
+		return Type;
+	}
+
+	// Set this entry to use an asset (rather than a struct)
+	ANIMNEXTUNCOOKEDONLY_API void SetAsset(const UAnimNextSharedVariables* InAsset, bool bSetupUndoRedo = true);
+
+	// Get the asset that this entry represents. If this entry uses a struct, this call will return nullptr
+	ANIMNEXTUNCOOKEDONLY_API const UAnimNextSharedVariables* GetAsset() const;
+
+	// Get the path to the asset/struct path that this entry represents, if any
+	ANIMNEXTUNCOOKEDONLY_API FSoftObjectPath GetObjectPath() const;
+
+	// Set this entry to use a struct (rather than an asset)
+	ANIMNEXTUNCOOKEDONLY_API void SetStruct(const UScriptStruct* InStruct, bool bSetupUndoRedo = true);
+
+	// Get the struct that this entry represents. If this entry uses an asset, this call will return nullptr
+	ANIMNEXTUNCOOKEDONLY_API const UScriptStruct* GetStruct() const;
+
+	// Recompiles this asset when the linked asset is modified
+	void HandleAssetModified(UAnimNextRigVMAssetEditorData* InEditorData, EAnimNextEditorDataNotifType InType, UObject* InSubject);
+
+#if WITH_LIVE_CODING
+	// Recompiles this asset when live coding is run (potentially modifying the struct we reference)
+	void HandlePatchComplete();
+#endif
+
+	/** The asset whose variables we share */
+	UPROPERTY(VisibleAnywhere, Category = "Asset")
+	TObjectPtr<const UAnimNextSharedVariables> Asset;
+
+	/** Soft reference to the asset/struct for error reporting */
+	UPROPERTY()
+	FSoftObjectPath ObjectPath;
+
+	/** The struct whose variables we use */
+	UPROPERTY(VisibleAnywhere, Category = "Asset")
+	TObjectPtr<const UScriptStruct> Struct;
+
+	UPROPERTY()
+	EAnimNextSharedVariablesType Type;
+};
diff --git a/Engine/Plugins/Experimental/AnimNext/Source/AnimNextUncookedOnly/Internal/Entries/AnimNextVariableEntry.h b/Engine/Plugins/Experimental/AnimNext/Source/AnimNextUncookedOnly/Internal/Entries/AnimNextVariableEntry.h
--- a/Engine/Plugins/Experimental/AnimNext/Source/AnimNextUncookedOnly/Internal/Entries/AnimNextVariableEntry.h
+++ b/Engine/Plugins/Experimental/AnimNext/Source/AnimNextUncookedOnly/Internal/Entries/AnimNextVariableEntry.h
@@ -14,9 +14,9 @@
 
 class UAnimNextRigVMAssetEditorData;
 class UAnimNextModule_EditorData;
-class UAnimNextDataInterface_EditorData;
+class UAnimNextSharedVariables_EditorData;
 class UAnimNextVariableBinding;
-class UAnimNextDataInterfaceEntry;
+class UAnimNextSharedVariablesEntry;
 class FAnimationAnimNextRuntimeTest_Variables;
 
 namespace UE::AnimNext::Editor
@@ -53,8 +53,8 @@ class UAnimNextVariableEntry : public UAnimNextRigVMAssetEntry, public IAnimNext
 
 	friend class UAnimNextRigVMAssetEditorData;
 	friend class UAnimNextModule_EditorData;
-	friend class UAnimNextDataInterface_EditorData;
-	friend class UAnimNextDataInterfaceEntry;
+	friend class UAnimNextSharedVariables_EditorData;
+	friend class UAnimNextSharedVariablesEntry;
 	friend class UAnimNextStateTreeTreeEditorData;
 	friend class UE::AnimNext::Tests::FEditor_Variables;
 	friend class UE::AnimNext::Tests::FEditor_AnimGraph_Variables;
diff --git a/Engine/Plugins/Experimental/AnimNext/Source/AnimNextUncookedOnly/Internal/Module/AnimNextModule_EditorData.h b/Engine/Plugins/Experimental/AnimNext/Source/AnimNextUncookedOnly/Internal/Module/AnimNextModule_EditorData.h
--- a/Engine/Plugins/Experimental/AnimNext/Source/AnimNextUncookedOnly/Internal/Module/AnimNextModule_EditorData.h
+++ b/Engine/Plugins/Experimental/AnimNext/Source/AnimNextUncookedOnly/Internal/Module/AnimNextModule_EditorData.h
@@ -5,7 +5,7 @@
 #include "CoreMinimal.h"
 #include "RigVMModel/RigVMGraph.h"
 #include "AnimNextEdGraph.h"
-#include "DataInterface/AnimNextDataInterface_EditorData.h"
+#include "Variables/AnimNextSharedVariables_EditorData.h"
 #include "AnimNextExecuteContext.h"
 #include "Module/RigVMTrait_ModuleEventDependency.h"
 #include "AnimNextModule_EditorData.generated.h"
@@ -31,7 +31,7 @@ namespace UE::AnimNext::Editor
 
 /** Editor data for AnimNext modules */
 UCLASS(MinimalAPI)
-class UAnimNextModule_EditorData : public UAnimNextDataInterface_EditorData
+class UAnimNextModule_EditorData : public UAnimNextSharedVariables_EditorData
 {
 	GENERATED_BODY()
 
diff --git a/Engine/Plugins/Experimental/AnimNext/Source/AnimNextUncookedOnly/Internal/UncookedOnlyUtils.h b/Engine/Plugins/Experimental/AnimNext/Source/AnimNextUncookedOnly/Internal/UncookedOnlyUtils.h
--- a/Engine/Plugins/Experimental/AnimNext/Source/AnimNextUncookedOnly/Internal/UncookedOnlyUtils.h
+++ b/Engine/Plugins/Experimental/AnimNext/Source/AnimNextUncookedOnly/Internal/UncookedOnlyUtils.h
@@ -17,6 +17,9 @@
 
 #define UE_API ANIMNEXTUNCOOKEDONLY_API
 
+class UAnimNextSharedVariables;
+struct FAnimNextSoftVariableReference;
+struct FAnimNextVariableReference;
 struct FAnimNextGetVariableCompileContext;
 struct FAnimNextVariableBindingData;
 struct FEdGraphPinType;
@@ -75,16 +78,6 @@ struct FAnimNextAssetRegistryExportedVariable
 		, Flags((uint32)InFlags) 
 	{}
 
-	bool operator==(const FAnimNextAssetRegistryExportedVariable& Other) const
-	{
-		return Name == Other.Name;
-	}
-
-	friend uint32 GetTypeHash(const FAnimNextAssetRegistryExportedVariable& Entry)
-	{
-		return GetTypeHash(Entry.Name);
-	}
-
 	EAnimNextExportedVariableFlags GetFlags() const
 	{
 		return (EAnimNextExportedVariableFlags)Flags;
@@ -100,6 +93,36 @@ struct FAnimNextAssetRegistryExportedVariable
 	uint32 Flags = (int32)EAnimNextExportedVariableFlags::NoFlags;
 };
 
+USTRUCT()
+struct FAnimNextAssetRegistryExportedSharedVariableAsset
+{
+	GENERATED_BODY()
+
+	FAnimNextAssetRegistryExportedSharedVariableAsset() = default;
+
+	FAnimNextAssetRegistryExportedSharedVariableAsset(const UAnimNextSharedVariables* InSharedVariables)
+		: SharedVariablesAsset(InSharedVariables)
+	{}
+
+	UPROPERTY()
+	TObjectPtr<const UAnimNextSharedVariables> SharedVariablesAsset;
+};
+
+USTRUCT()
+struct FAnimNextAssetRegistryExportedSharedVariableStruct
+{
+	GENERATED_BODY()
+
+	FAnimNextAssetRegistryExportedSharedVariableStruct() = default;
+
+	FAnimNextAssetRegistryExportedSharedVariableStruct(const UScriptStruct* InStruct)
+		: SharedVariablesStruct(InStruct)
+	{}
+
+	UPROPERTY()
+	TObjectPtr<const UScriptStruct> SharedVariablesStruct;
+};
+
 USTRUCT()
 struct FAnimNextAssetRegistryManifestNode
 {
@@ -144,6 +167,12 @@ struct FAnimNextAssetRegistryExports
 	UPROPERTY()
 	TArray<FAnimNextAssetRegistryExportedVariable> Variables;
 
+	UPROPERTY()
+	TArray<FAnimNextAssetRegistryExportedSharedVariableAsset> SharedVariableAssets;
+
+	UPROPERTY()
+	TArray<FAnimNextAssetRegistryExportedSharedVariableStruct> SharedVariableStructs;
+
 	UPROPERTY()
 	TArray<FAnimNextAssetRegistryManifestNode> ManifestNodes;
 };
@@ -216,7 +245,7 @@ struct FUtils
 
 	// Gets the exported variables that are used by a RigVM asset
 	static UE_API void GetAssetVariables(const UAnimNextRigVMAssetEditorData* EditorData, FAnimNextAssetRegistryExports& OutExports);
-	static UE_API void GetAssetVariables(const UAnimNextRigVMAssetEditorData* EditorData, TSet<FAnimNextAssetRegistryExportedVariable>& OutExports);
+	static UE_API void GetAssetVariables(const UAnimNextRigVMAssetEditorData* EditorData, TSet<FAnimNextAssetRegistryExportedVariable>& OutExports, TArray<FAnimNextAssetRegistryExportedSharedVariableAsset>& OutExportedSharedVariableAssets, TArray<FAnimNextAssetRegistryExportedSharedVariableStruct>& OutExportedSharedVariableStructs);
 
 	// Gets the exported public functions that are used by a RigVM asset
 	static UE_API void GetAssetFunctions(const UAnimNextRigVMAssetEditorData* InEditorData, FRigVMGraphFunctionHeaderArray& OutExports);
@@ -231,10 +260,11 @@ struct FUtils
 	static UE_API void CreateFunctionLibraryOutlinerItemsRecursive(const UAnimNextRigVMAssetEditorData* EditorData, FWorkspaceOutlinerItemExports& OutExports, FWorkspaceOutlinerItemExport& RootExport, int32 ParentExportIndex, const TArray<FRigVMGraphFunctionData>& PublicFunctions, const TArray<FRigVMGraphFunctionData>& PrivateFunctions);
 	static UE_API void CreateFunctionsOutlinerItemsRecursive(const UAnimNextRigVMAssetEditorData* EditorData, FWorkspaceOutlinerItemExports& OutExports, FWorkspaceOutlinerItemExport& RootExport, int32 ParentExportIndex, const TArray<FRigVMGraphFunctionData>& Functions, bool bPublicFunctions);
 
-	// Attempts to determine the type from a parameter name
+	// Attempts to determine the type from a variable name & asset/struct path
 	// If the name cannot be found, the returned type will be invalid
 	// Note that this is expensive and can query the asset registry
-	static UE_API FAnimNextParamType GetParameterTypeFromName(FName InName);
+	static UE_API FAnimNextParamType FindVariableType(const FAnimNextVariableReference& InVariableReference);
+	static UE_API FAnimNextParamType FindVariableType(const FAnimNextSoftVariableReference& InSoftVariableReference);
 
 	// Returns an user friendly name for the Function Library
 	static UE_API const FText& GetFunctionLibraryDisplayName();
diff --git a/Engine/Plugins/Experimental/AnimNext/Source/AnimNextUncookedOnly/Internal/Variables/AnimNextSharedVariables_EditorData.h b/Engine/Plugins/Experimental/AnimNext/Source/AnimNextUncookedOnly/Internal/Variables/AnimNextSharedVariables_EditorData.h
new file mode 100644
--- /dev/null
+++ b/Engine/Plugins/Experimental/AnimNext/Source/AnimNextUncookedOnly/Internal/Variables/AnimNextSharedVariables_EditorData.h
@@ -0,0 +1,31 @@
+// Copyright Epic Games, Inc. All Rights Reserved.
+
+#pragma once
+
+#include "CoreMinimal.h"
+#include "AnimNextExecuteContext.h"
+#include "AnimNextRigVMAssetEditorData.h"
+#include "AnimNextSharedVariables_EditorData.generated.h"
+
+class UAnimNextSharedVariablesFactory;
+
+namespace UE::AnimNext::UncookedOnly
+{
+	struct FUtils;
+}
+
+/** Editor data for AnimNext data interfaces */
+UCLASS()
+class ANIMNEXTUNCOOKEDONLY_API UAnimNextSharedVariables_EditorData : public UAnimNextRigVMAssetEditorData
+{
+	GENERATED_BODY()
+
+protected:
+	// UAnimNextRigVMAssetEditorData interface
+	virtual UScriptStruct* GetExecuteContextStruct() const override { return FAnimNextExecuteContext::StaticStruct(); }
+	virtual TConstArrayView<TSubclassOf<UAnimNextRigVMAssetEntry>> GetEntryClasses() const override;
+	virtual void CustomizeNewAssetEntry(UAnimNextRigVMAssetEntry* InNewEntry) const override;
+
+	friend class UAnimNextSharedVariablesFactory;
+	friend UE::AnimNext::UncookedOnly::FUtils;
+};
diff --git a/Engine/Plugins/Experimental/AnimNext/Source/AnimNextUncookedOnly/Internal/Variables/K2Node_AnimNextComponentSetVariable.h b/Engine/Plugins/Experimental/AnimNext/Source/AnimNextUncookedOnly/Internal/Variables/K2Node_AnimNextComponentSetVariable.h
new file mode 100644
--- /dev/null
+++ b/Engine/Plugins/Experimental/AnimNext/Source/AnimNextUncookedOnly/Internal/Variables/K2Node_AnimNextComponentSetVariable.h
@@ -0,0 +1,35 @@
+﻿// Copyright Epic Games, Inc. All Rights Reserved.
+
+#pragma once
+
+#include "K2Node_CallFunction.h"
+#include "K2Node_AnimNextComponentSetVariable.generated.h"
+
+/** A custom node that supports SetVariable functionality */
+UCLASS(MinimalAPI)
+class UK2Node_AnimNextComponentSetVariable : public UK2Node_CallFunction
+{
+	GENERATED_BODY()
+
+public:
+	ANIMNEXTUNCOOKEDONLY_API UK2Node_AnimNextComponentSetVariable();
+
+private:
+	// //~ Begin UEdGraphNode Interface.
+	ANIMNEXTUNCOOKEDONLY_API virtual void ReconstructNode() override;
+	ANIMNEXTUNCOOKEDONLY_API virtual void PinDefaultValueChanged(UEdGraphPin* Pin) override;
+	// //~ End UEdGraphNode Interface.
+
+	//~ Begin K2Node Interface
+	ANIMNEXTUNCOOKEDONLY_API virtual void GetMenuActions(FBlueprintActionDatabaseRegistrar& InActionRegistrar) const override;
+	ANIMNEXTUNCOOKEDONLY_API virtual void ValidateNodeDuringCompilation(FCompilerResultsLog& MessageLog) const override;
+	ANIMNEXTUNCOOKEDONLY_API virtual void PreloadRequiredAssets() override;
+	//~ End K2Node Interface
+
+	void CachePinType();
+	void PropagateCachedTypeToPin();
+
+	// Cached type of the pin used for regenerating after wildcard is reconstructed
+	UPROPERTY()
+	FEdGraphPinType CachedPinType;
+};
diff --git a/Engine/Plugins/Experimental/AnimNext/Source/AnimNextUncookedOnly/Private/AnimNextAssetWorkspaceAssetUserData.cpp b/Engine/Plugins/Experimental/AnimNext/Source/AnimNextUncookedOnly/Private/AnimNextAssetWorkspaceAssetUserData.cpp
--- a/Engine/Plugins/Experimental/AnimNext/Source/AnimNextUncookedOnly/Private/AnimNextAssetWorkspaceAssetUserData.cpp
+++ b/Engine/Plugins/Experimental/AnimNext/Source/AnimNextUncookedOnly/Private/AnimNextAssetWorkspaceAssetUserData.cpp
@@ -3,7 +3,7 @@
 #include "AnimNextAssetWorkspaceAssetUserData.h"
 
 #include "UncookedOnlyUtils.h"
-#include "DataInterface/AnimNextDataInterface.h"
+#include "Variables/AnimNextSharedVariables.h"
 #include "Module/AnimNextModule.h"
 #include "Module/AnimNextModule_EditorData.h"
 #include "UObject/AssetRegistryTagsContext.h"
@@ -22,11 +22,11 @@ void UAnimNextAssetWorkspaceAssetUserData::GetAssetRegistryTags(FAssetRegistryTa
 			{
 				FWorkspaceOutlinerItemExport& RootAssetExport = CachedExports.Exports.Add_GetRef(FWorkspaceOutlinerItemExport(Asset->GetFName(), Asset));
 
-				if(UAnimNextModule* Module = Cast<UAnimNextModule>(Asset))
+				if(Asset->IsA<UAnimNextModule>())
 				{
 					RootAssetExport.GetData().InitializeAsScriptStruct(FAnimNextModuleOutlinerData::StaticStruct());
 				}
-				else if(UAnimNextDataInterface* DataInterface = Cast<UAnimNextDataInterface>(Asset))
+				else if(Asset->IsA<UAnimNextSharedVariables>())
 				{
 					RootAssetExport.GetData().InitializeAsScriptStruct(FAnimNextDataInterfaceOutlinerData::StaticStruct());
 				}
diff --git a/Engine/Plugins/Experimental/AnimNext/Source/AnimNextUncookedOnly/Private/AnimNextControllerBase.cpp b/Engine/Plugins/Experimental/AnimNext/Source/AnimNextUncookedOnly/Private/AnimNextControllerBase.cpp
new file mode 100644
--- /dev/null
+++ b/Engine/Plugins/Experimental/AnimNext/Source/AnimNextUncookedOnly/Private/AnimNextControllerBase.cpp
@@ -0,0 +1,86 @@
+﻿// Copyright Epic Games, Inc. All Rights Reserved.
+
+#include "AnimNextControllerBase.h"
+#include "AnimNextSharedVariableNode.h"
+#include "Variables/AnimNextSharedVariables.h"
+
+UAnimNextSharedVariableNode* UAnimNextControllerBase::AddAssetSharedVariableNode(const UAnimNextSharedVariables* InAsset, const FName& InVariableName, const FString& InCPPType, UObject* InCPPTypeObject, bool bIsGetter, const FString& InDefaultValue, const FVector2D& InPosition, const FString& InNodeName, bool bSetupUndoRedo, bool bPrintPythonCommand)
+{
+	if (InAsset == nullptr)
+	{
+		ReportError(TEXT("Invalid asset supplied to AddAssetSharedVariableNode."));
+		return nullptr;
+	}
+
+	return AddSharedVariableNode(InAsset->GetPathName(), InVariableName, InCPPType, InCPPTypeObject, bIsGetter, InDefaultValue, InPosition, InNodeName, bSetupUndoRedo, bPrintPythonCommand);
+}
+
+UAnimNextSharedVariableNode* UAnimNextControllerBase::AddStructSharedVariableNode(const UScriptStruct* InStruct, const FName& InVariableName, const FString& InCPPType, UObject* InCPPTypeObject, bool bIsGetter, const FString& InDefaultValue, const FVector2D& InPosition, const FString& InNodeName, bool bSetupUndoRedo, bool bPrintPythonCommand)
+{
+	if (InStruct == nullptr)
+	{
+		ReportError(TEXT("Invalid struct supplied to AddStructSharedVariableNode."));
+		return nullptr;
+	}
+
+	return AddSharedVariableNode(InStruct->GetPathName(), InVariableName, InCPPType, InCPPTypeObject, bIsGetter, InDefaultValue, InPosition, InNodeName, bSetupUndoRedo, bPrintPythonCommand);
+}
+
+UAnimNextSharedVariableNode* UAnimNextControllerBase::AddSharedVariableNode(const FString& InSourceObjectPath, const FName& InVariableName, const FString& InCPPType, UObject* InCPPTypeObject, bool bIsGetter, const FString& InDefaultValue, const FVector2D& InPosition, const FString& InNodeName, bool bSetupUndoRedo, bool bPrintPythonCommand)
+{
+	UObject* SourceObject = nullptr;
+	if (!InSourceObjectPath.IsEmpty())
+	{
+		SourceObject = RigVMTypeUtils::FindObjectFromCPPTypeObjectPath<UObject>(InSourceObjectPath);
+		if (SourceObject == nullptr)
+		{
+			return nullptr;
+		}
+	}
+
+	UAnimNextSharedVariableNode* VariableNode = CastChecked<UAnimNextSharedVariableNode>(AddVariableNode(InVariableName, UAnimNextSharedVariableNode::StaticClass(), InCPPType, InCPPTypeObject, bIsGetter, InDefaultValue, InPosition, InNodeName, bSetupUndoRedo, bPrintPythonCommand), ECastCheckedType::Type::NullAllowed);
+	if (VariableNode == nullptr)
+	{
+		return nullptr;
+	}
+
+	if (const UAnimNextSharedVariables* Asset = Cast<UAnimNextSharedVariables>(SourceObject))
+	{
+		VariableNode->Type = EAnimNextSharedVariablesType::Asset;
+		VariableNode->Asset = Asset;
+	}
+	else if (const UScriptStruct* Struct = Cast<UScriptStruct>(SourceObject))
+	{
+		VariableNode->Type = EAnimNextSharedVariablesType::Struct;
+		VariableNode->Struct = Struct;
+	}
+	else
+	{
+		return nullptr;
+	}
+
+	return VariableNode;
+}
+
+void UAnimNextControllerBase::ReplaceVariableNodeWithSharedVariableNode(URigVMVariableNode* InVariableNode, FName InNewVariableName, const UObject* InAssetOrStruct)
+{
+	check(InVariableNode);
+	check(InAssetOrStruct);
+	check(InNewVariableName != NAME_None);
+
+	UAnimNextSharedVariableNode* SharedVariableNode = AddSharedVariableNode(InAssetOrStruct->GetPathName(), InNewVariableName, InVariableNode->GetCPPType(), InVariableNode->GetCPPTypeObject(), InVariableNode->IsGetter(), InVariableNode->GetDefaultValue(), InVariableNode->GetPosition(), FString(), false, false);
+	if (SharedVariableNode)
+	{
+		if (!InVariableNode->IsGetter())
+		{
+			RewireLinks(InVariableNode->GetValuePin(), SharedVariableNode->GetValuePin(), true, false);
+			RewireLinks(InVariableNode->FindExecutePin(), SharedVariableNode->FindExecutePin(), true, false);
+			RewireLinks(InVariableNode->FindExecutePin(), SharedVariableNode->FindExecutePin(), false, false);
+		}
+		else
+		{
+			RewireLinks(InVariableNode->GetValuePin(), SharedVariableNode->GetValuePin(), false, false);
+		}
+		RemoveNode(InVariableNode, false, false);
+	}
+}
\ No newline at end of file
diff --git a/Engine/Plugins/Experimental/AnimNext/Source/AnimNextUncookedOnly/Private/AnimNextRigVMAssetEditorData.cpp b/Engine/Plugins/Experimental/AnimNext/Source/AnimNextUncookedOnly/Private/AnimNextRigVMAssetEditorData.cpp
--- a/Engine/Plugins/Experimental/AnimNext/Source/AnimNextUncookedOnly/Private/AnimNextRigVMAssetEditorData.cpp
+++ b/Engine/Plugins/Experimental/AnimNext/Source/AnimNextUncookedOnly/Private/AnimNextRigVMAssetEditorData.cpp
@@ -12,6 +12,7 @@
 #include "Entries/AnimNextRigVMAssetEntry.h"
 #include "AnimNextRigVMAssetSchema.h"
 #include "AnimNextAssetWorkspaceAssetUserData.h"
+#include "AnimNextControllerBase.h"
 #include "AnimNextScopedCompilerResults.h"
 #include "RigVMPythonUtils.h"
 #include "ExternalPackageHelper.h"
@@ -22,8 +23,11 @@
 #include "UncookedOnlyUtils.h"
 #include "Animation/AnimCompressionTypes.h"
 #include "Animation/Skeleton.h"
-#include "DataInterface/AnimNextDataInterface_EditorData.h"
-#include "Entries/AnimNextDataInterfaceEntry.h"
+#include "AnimNext/Private/AnimNextRigVMFunctionData.h"
+#include "AnimNext/Private/Variables/StructDataCache.h"
+#include "DataInterface/AnimNextNativeDataInterface.h"
+#include "Variables/AnimNextSharedVariables_EditorData.h"
+#include "Entries/AnimNextSharedVariablesEntry.h"
 #include "Entries/AnimNextEventGraphEntry.h"
 #include "Entries/AnimNextVariableEntry.h"
 #include "Graph/RigUnit_AnimNextBeginExecution.h"
@@ -35,10 +39,11 @@
 #include "RigVMModel/Nodes/RigVMAggregateNode.h"
 #include "RigVMModel/Nodes/RigVMCollapseNode.h"
 #include "UObject/AssetRegistryTagsContext.h"
-#include "DataInterface/AnimNextDataInterface_EditorData.h"
+#include "Variables/AnimNextSharedVariables_EditorData.h"
 #include "RigVMFunctions/Execution/RigVMFunction_UserDefinedEvent.h"
 #include "Misc/UObjectToken.h"
 #include "UObject/SavePackage.h"
+#include "UObject/FortniteMainBranchObjectVersion.h"
 
 #if WITH_EDITOR
 #include "FileHelpers.h"
@@ -96,6 +101,8 @@ void UAnimNextRigVMAssetEditorData::ReconstructAllNodes()
 
 void UAnimNextRigVMAssetEditorData::Serialize(FArchive& Ar)
 {
+	Ar.UsingCustomVersion(FFortniteMainBranchObjectVersion::GUID);
+
 	RigVMClient.SetDefaultSchemaClass(UAnimNextRigVMAssetSchema::StaticClass());
 	RigVMClient.SetOuterClientHost(this, GET_MEMBER_NAME_CHECKED(UAnimNextRigVMAssetEditorData, RigVMClient));
 
@@ -106,6 +113,23 @@ void UAnimNextRigVMAssetEditorData::Serialize(FArchive& Ar)
 	}
 
 	Super::Serialize(Ar);
+
+	if (Ar.IsLoading() && Ar.CustomVer(FFortniteMainBranchObjectVersion::GUID) < FFortniteMainBranchObjectVersion::AnimNextVariableReferences)
+	{
+		PRAGMA_DISABLE_DEPRECATION_WARNINGS
+		DefaultInjectionSiteReference = FAnimNextVariableReference(DefaultInjectionSite_DEPRECATED);
+		PRAGMA_ENABLE_DEPRECATION_WARNINGS
+
+		PRAGMA_DISABLE_DEPRECATION_WARNINGS
+		if (NativeInterface_DEPRECATED)
+		{
+			NativeInterfaces_DEPRECATED.Add(NativeInterface_DEPRECATED);
+			NativeInterface_DEPRECATED = nullptr;
+		}
+		PRAGMA_ENABLE_DEPRECATION_WARNINGS
+
+		bUpgradeDataInterfacesOnLoad = true;
+	}
 }
 
 void UAnimNextRigVMAssetEditorData::Initialize(bool bRecompileVM)
@@ -182,15 +206,6 @@ void UAnimNextRigVMAssetEditorData::PostLoad()
 {
 	Super::PostLoad();
 
-	// Handle deprecation
-	PRAGMA_DISABLE_DEPRECATION_WARNINGS
-	if (NativeInterface_DEPRECATED)
-	{
-		NativeInterfaces.Add(NativeInterface_DEPRECATED);
-		NativeInterface_DEPRECATED = nullptr;
-	}
-	PRAGMA_ENABLE_DEPRECATION_WARNINGS
-
 	GraphModels.Reset();
 	
 	PostLoadExternalPackages();
@@ -345,6 +360,11 @@ void UAnimNextRigVMAssetEditorData::HandlePackageDone()
 
 	TGuardValue<bool> DisableCompilationNotifications(bSuspendCompilationNotifications, true);
 
+	if (bUpgradeDataInterfacesOnLoad)
+	{
+		UpgradeDataInterfaces();
+	}
+
 	RecompileVM();
 }
 
@@ -771,6 +791,191 @@ TObjectPtr<URigVMGraph> UAnimNextRigVMAssetEditorData::CreateContainedGraphModel
 	return Model;
 }
 
+void UAnimNextRigVMAssetEditorData::BuildFunctionWrapperEventVariables(FAnimNextRigVMAssetCompileContext& InContext)
+{
+	using namespace UE::AnimNext::UncookedOnly;
+	
+	for (const FAnimNextProgrammaticFunctionHeader& ProgrammaticFunctionHeader : InContext.FunctionHeaders)
+	{
+		const FRigVMGraphFunctionHeader& FunctionHeader = ProgrammaticFunctionHeader.FunctionHeader;
+		for (const FRigVMGraphFunctionArgument& Argument : FunctionHeader.Arguments)
+		{
+			FRigVMGraphFunctionArgument InternallyNamedArgument = Argument;
+			InternallyNamedArgument.Name = FName(FUtils::MakeFunctionWrapperVariableName(FunctionHeader.Name, Argument.Name));
+			InContext.ProgrammaticVariables.Add(FAnimNextProgrammaticVariable::FromRigVMGraphFunctionArgument(InternallyNamedArgument));
+		}
+	}
+}
+
+void UAnimNextRigVMAssetEditorData::BuildFunctionWrapperEvents(FAnimNextRigVMAssetCompileContext& InContext, const FRigVMCompileSettings& InSettings)
+{
+	using namespace UE::AnimNext::UncookedOnly;
+
+	UAnimNextRigVMAsset* Asset = FUtils::GetAsset<UAnimNextRigVMAsset>(this);
+	Asset->FunctionData.Reset();
+
+	FRigVMClient* VMClient = GetRigVMClient();
+
+	// Create all shim events for our traits to call
+	const bool bSetupUndoRedo = false;
+	URigVMGraph* WrapperGraph = NewObject<URigVMGraph>(this, NAME_None, RF_Transient);
+	URigVMController* Controller = VMClient->GetOrCreateController(WrapperGraph);
+	FRigVMControllerNotifGuard NotifGuard(Controller);
+	bool bAddedWrapperEvent = true;
+	TArray<FRigVMExternalVariable> ExternalVariables = Asset->GetExternalVariables();
+
+	for(const FAnimNextProgrammaticFunctionHeader& AnimNextFunctionHeader : InContext.FunctionHeaders)
+	{
+		const FRigVMGraphFunctionHeader& FunctionHeader = AnimNextFunctionHeader.FunctionHeader;
+
+		URigVMLibraryNode* LibraryNode = Cast<URigVMLibraryNode>(FunctionHeader.LibraryPointer.GetNodeSoftPath().TryLoad());
+		if(LibraryNode == nullptr)
+		{
+			InSettings.ReportError(FString::Printf(TEXT("Could not find function '%s'"), *FunctionHeader.Name.ToString()));
+			continue;
+		}
+
+		// Create user-defined entry point
+		FString WrapperEventName = FUtils::MakeFunctionWrapperEventName(FunctionHeader.Name);
+		URigVMUnitNode* EventNode = Controller->AddUnitNode(FRigVMFunction_UserDefinedEvent::StaticStruct(), TEXT("Execute"), FVector2D::ZeroVector, FunctionHeader.Name.ToString(), bSetupUndoRedo);
+		if(EventNode == nullptr)
+		{
+			InSettings.ReportError(FString::Printf(TEXT("Could not spawn event node for function '%s'"), *FunctionHeader.Name.ToString()));
+			continue;
+		}
+		URigVMPin* EventNamePin = EventNode->FindPin(GET_MEMBER_NAME_STRING_CHECKED(FRigVMFunction_UserDefinedEvent, EventName));
+		if(EventNamePin == nullptr)
+		{
+			InSettings.ReportError(FString::Printf(TEXT("Could not find custom event name pin")));
+			continue;
+		}
+		Controller->SetPinDefaultValue(EventNamePin->GetPinPath(), WrapperEventName, true, bSetupUndoRedo);
+
+		// Call function
+		URigVMFunctionReferenceNode* FunctionNode = Controller->AddFunctionReferenceNode(LibraryNode, FVector2D::ZeroVector, FunctionHeader.Name.ToString(), bSetupUndoRedo);
+		if(FunctionNode == nullptr)
+		{
+			InSettings.ReportError(FString::Printf(TEXT("Could not spawn function node for function '%s'"), *FunctionHeader.Name.ToString()));
+			continue;
+		}
+
+		// Link up Execute nodes if needed, function may be pure & lack an input pin
+		URigVMPin* CurrentExecuteOutputPin = EventNode->FindPin(FRigVMStruct::ExecuteContextName.ToString());
+		URigVMPin* ExecuteInputPin = FunctionNode->FindPin(FRigVMStruct::ExecuteContextName.ToString());
+		if (ExecuteInputPin && !Controller->AddLink(CurrentExecuteOutputPin, ExecuteInputPin, bSetupUndoRedo))
+		{
+			InSettings.ReportError(FString::Printf(TEXT("Could not link execute pins for function '%s'"), *FunctionHeader.Name.ToString()));
+			continue;
+		}
+
+		// Update current execute pin, RigVM doesn't have a concept of input / output execute pins, just one execute content pin used for both
+		CurrentExecuteOutputPin = ExecuteInputPin ? ExecuteInputPin : CurrentExecuteOutputPin;
+
+			// Controller needs to notify the AST of variable changes to make new links
+		bool bSuspendNotificationForInternalVariables = false;
+		FRigVMControllerNotifGuard VarNotifGuard(Controller, bSuspendNotificationForInternalVariables);
+
+		// Generate & link input arguments, also generate result variable node but link later
+		TArray<int32> ArgIndices;
+		for (const FRigVMGraphFunctionArgument& Argument : FunctionHeader.Arguments)
+		{
+			// Execution context is captured as arg pins, skip those for internal variable gen
+			if (Argument.Direction == ERigVMPinDirection::IO)
+			{
+				continue;
+			}
+
+			bool bIsGetter = Argument.Direction == ERigVMPinDirection::Input;
+			FName InternalArgName = FName(FUtils::MakeFunctionWrapperVariableName(FunctionHeader.Name, Argument.Name));
+			int32 VariableIndex = ExternalVariables.IndexOfByPredicate([&InternalArgName](const FRigVMExternalVariable& InVariable)
+			{
+				return InVariable.Name == InternalArgName;
+			});
+			check(VariableIndex != INDEX_NONE);
+			ArgIndices.Add(VariableIndex);
+
+			if (bIsGetter)
+			{
+				URigVMVariableNode* FunctionParamVariableNode = Controller->AddVariableNode(InternalArgName
+					, Argument.CPPType.ToString()
+					, Argument.CPPTypeObject.Get()
+					, bIsGetter
+					, Argument.DefaultValue
+					, FVector2D::ZeroVector
+					, InternalArgName.ToString()
+					, bSetupUndoRedo);
+
+				if (!FunctionParamVariableNode)
+				{
+					InSettings.ReportError(FString::Printf(TEXT("Failed to add internal variable node for param: %s, var: %s"), *FunctionHeader.Name.ToString(), *InternalArgName.ToString()));
+					return;
+				}
+
+				// Link Param Pins
+				URigVMPin* ParamValuePin = FunctionParamVariableNode->GetValuePin();
+				URigVMPin* FunctionArgumentPin = FunctionNode->FindPin(Argument.Name.ToString());
+				if (!Controller->AddLink(ParamValuePin, FunctionArgumentPin, bSetupUndoRedo))
+				{
+					InSettings.ReportError(FString::Printf(TEXT("Failed to link internal variable param node to function: %s -> %s"), *GetNameSafe(ParamValuePin), *GetNameSafe(FunctionArgumentPin)));
+					return;
+				}
+			}
+
+			if (!bIsGetter)
+			{
+				
+				URigVMVariableNode* FunctionResultVariableNode = Controller->AddVariableNode(InternalArgName
+					, Argument.CPPType.ToString()
+					, Argument.CPPTypeObject.Get()
+					, bIsGetter
+					, Argument.DefaultValue
+					, FVector2D::ZeroVector
+					, InternalArgName.ToString()
+					, bSetupUndoRedo);
+
+				if (!FunctionResultVariableNode)
+				{
+					InSettings.ReportError(FString::Printf(TEXT("Failed to add internal variable node for result: %s, var: %s"), *FunctionHeader.Name.ToString(), *InternalArgName.ToString()));
+					return;
+				}
+
+				// Link Result pins
+				URigVMPin* FunctionResultPin = FunctionNode->FindPin(Argument.Name.ToString());
+				URigVMPin* ResultValuePin = FunctionResultVariableNode->GetValuePin();
+				if (!Controller->AddLink(FunctionResultPin, ResultValuePin, bSetupUndoRedo))
+				{
+					InSettings.ReportError(FString::Printf(TEXT("Failed to link internal variable result node to function: %s -> %s"), *GetNameSafe(FunctionResultPin), *GetNameSafe(ResultValuePin)));
+					return;
+				}
+
+				// Link Result Execute pins
+				URigVMPin* ResultExecuteInputPin = FunctionResultVariableNode->FindPin(FRigVMStruct::ExecuteContextName.ToString());
+				if (!Controller->AddLink(CurrentExecuteOutputPin, ResultExecuteInputPin, bSetupUndoRedo))
+				{
+					InSettings.ReportError(FString::Printf(TEXT("Failed to link execute pins for variable result node: %s -> %s"), *GetNameSafe(CurrentExecuteOutputPin), *GetNameSafe(ResultExecuteInputPin)));
+					return;
+				}
+
+				// Update current execute pin, RigVM doesn't have a concept of input / output execute pins, just one execute content pin used for both
+				CurrentExecuteOutputPin = ResultExecuteInputPin;
+			}
+		}
+
+		bAddedWrapperEvent = true;
+
+		FAnimNextRigVMFunctionData FunctionData;
+		FunctionData.Name = FunctionHeader.Name;
+		FunctionData.EventName = FName(*WrapperEventName);
+		FunctionData.ArgIndices = MoveTemp(ArgIndices);
+		Asset->FunctionData.Add(MoveTemp(FunctionData));
+	}
+
+	if(bAddedWrapperEvent)
+	{
+		InContext.ProgrammaticGraphs.Add(WrapperGraph);
+	}
+}
+
 void UAnimNextRigVMAssetEditorData::RecompileVM()
 {
 	using namespace UE::AnimNext::UncookedOnly;
@@ -815,11 +1020,15 @@ void UAnimNextRigVMAssetEditorData::RecompileVM()
 			OnPreCompileGetProgrammaticFunctionHeaders(Settings, GetFunctionHeaderCompileContext);
 		}
 
+		BuildFunctionWrapperEventVariables(CompileContext);
+
 		{
 			FAnimNextGetVariableCompileContext GetVariableCompileContext(CompileContext);
 			FUtils::CompileVariables(Settings, Asset, GetVariableCompileContext);
 		}
 
+		BuildFunctionWrapperEvents(CompileContext, Settings);
+
 		{
 			FAnimNextGetGraphCompileContext GetGraphCompileContext(CompileContext);
 			OnPreCompileGetProgrammaticGraphs(Settings, GetGraphCompileContext);
@@ -1015,11 +1224,11 @@ bool UAnimNextRigVMAssetEditorData::IsDirtyForRecompilation() const
 	}
 
 	bool bDependencyDirty = false;
-	ForEachEntryOfType<UAnimNextDataInterfaceEntry>([&bDependencyDirty](UAnimNextDataInterfaceEntry* InEntry)
+	ForEachEntryOfType<UAnimNextSharedVariablesEntry>([&bDependencyDirty](UAnimNextSharedVariablesEntry* InEntry)
 	{
-		if(InEntry->DataInterface)
+		if(InEntry->Asset)
 		{
-			UAnimNextRigVMAssetEditorData* EditorData = UE::AnimNext::UncookedOnly::FUtils::GetEditorData<UAnimNextRigVMAssetEditorData>(InEntry->DataInterface.Get());
+			UAnimNextRigVMAssetEditorData* EditorData = UE::AnimNext::UncookedOnly::FUtils::GetEditorData<UAnimNextRigVMAssetEditorData>(InEntry->Asset.Get());
 			if(EditorData->IsDirtyForRecompilation())
 			{
 				bDependencyDirty = true;
@@ -1862,41 +2071,35 @@ UAnimNextEventGraphEntry* UAnimNextRigVMAssetEditorData::AddEventGraph(FName InN
 	return NewEntry;
 }
 
-
-UAnimNextDataInterfaceEntry* UAnimNextRigVMAssetLibrary::AddDataInterface(UAnimNextRigVMAsset* InAsset, UAnimNextDataInterface* InDataInterface, bool bSetupUndoRedo, bool bPrintPythonCommand)
+UAnimNextSharedVariablesEntry* UAnimNextRigVMAssetLibrary::AddSharedVariables(UAnimNextRigVMAsset* InAsset, UAnimNextSharedVariables* InSharedVariables, bool bSetupUndoRedo, bool bPrintPythonCommand)
 {
-	return UE::AnimNext::UncookedOnly::FUtils::GetEditorData(InAsset)->AddDataInterface(InDataInterface, bSetupUndoRedo, bPrintPythonCommand);
+	return UE::AnimNext::UncookedOnly::FUtils::GetEditorData(InAsset)->AddSharedVariables(InSharedVariables, bSetupUndoRedo, bPrintPythonCommand);
 }
 
-URigVMLibraryNode* UAnimNextRigVMAssetLibrary::AddFunction(UAnimNextRigVMAsset* InAsset, FName InFunctionName, bool bInMutable, bool bSetupUndoRedo, bool bPrintPythonCommand)
+UAnimNextSharedVariablesEntry* UAnimNextRigVMAssetEditorData::AddSharedVariables(const UAnimNextSharedVariables* InSharedVariables, bool bSetupUndoRedo, bool bPrintPythonCommand)
 {
-	return UE::AnimNext::UncookedOnly::FUtils::GetEditorData(InAsset)->AddFunction(InFunctionName, bInMutable, bSetupUndoRedo, bPrintPythonCommand);
-}
-
-UAnimNextDataInterfaceEntry* UAnimNextRigVMAssetEditorData::AddDataInterface(UAnimNextDataInterface* InDataInterface, bool bSetupUndoRedo, bool bPrintPythonCommand)
-{
-	if(InDataInterface == nullptr)
+	if(InSharedVariables == nullptr)
 	{
-		ReportError(TEXT("UAnimNextRigVMAssetEditorData::AddDataInterface: Invalid data interface supplied."));
+		ReportError(TEXT("UAnimNextRigVMAssetEditorData::AddSharedVariables: Invalid asset supplied."));
 		return nullptr;
 	}
 
-	if(!GetEntryClasses().Contains(UAnimNextDataInterfaceEntry::StaticClass()) || !CanAddNewEntry(UAnimNextDataInterfaceEntry::StaticClass()))
+	if(!GetEntryClasses().Contains(UAnimNextSharedVariablesEntry::StaticClass()) || !CanAddNewEntry(UAnimNextSharedVariablesEntry::StaticClass()))
 	{
-		ReportError(TEXT("UAnimNextRigVMAssetEditorData::AddDataInterface: Cannot add a data interface to this asset - entry is not allowed."));
+		ReportError(TEXT("UAnimNextRigVMAssetEditorData::AddSharedVariables: Cannot add a shared variables to this asset - entry is not allowed."));
 		return nullptr;
 	}
 	
 	// Check if interface has any public members or if any of its parent interfaces do
-	UAnimNextDataInterface_EditorData* EditorData = UE::AnimNext::UncookedOnly::FUtils::GetEditorData<UAnimNextDataInterface_EditorData>(InDataInterface);
+	UAnimNextSharedVariables_EditorData* EditorData = UE::AnimNext::UncookedOnly::FUtils::GetEditorData<UAnimNextSharedVariables_EditorData>(InSharedVariables);
 	if(EditorData == nullptr)
 	{
-		ReportError(TEXT("UAnimNextRigVMAssetEditorData::AddDataInterface: Invalid data interface supplied - asset has no editor data."));
+		ReportError(TEXT("UAnimNextRigVMAssetEditorData::AddSharedVariables: Invalid asset supplied - asset has no editor data."));
 		return nullptr;
 	}
 
 	// Check for circularity
-	auto CheckForCircularity = [this](UAnimNextDataInterface_EditorData* InEditorData, auto& InCheckForCircularity)
+	auto CheckForCircularity = [this](UAnimNextSharedVariables_EditorData* InEditorData, auto& InCheckForCircularity)
 	{
 		if(InEditorData == this)
 		{
@@ -1905,10 +2108,10 @@ UAnimNextDataInterfaceEntry* UAnimNextRigVMAssetEditorData::AddDataInterface(UAn
 
 		for(UAnimNextRigVMAssetEntry* Entry : InEditorData->Entries)
 		{
-			if(UAnimNextDataInterfaceEntry* DataInterfaceEntry = Cast<UAnimNextDataInterfaceEntry>(Entry))
+			if(UAnimNextSharedVariablesEntry* SharedVariablesEntry = Cast<UAnimNextSharedVariablesEntry>(Entry))
 			{
-				UAnimNextDataInterface* DataInterface = DataInterfaceEntry->GetDataInterface();
-				UAnimNextDataInterface_EditorData* EditorData = UE::AnimNext::UncookedOnly::FUtils::GetEditorData<UAnimNextDataInterface_EditorData>(DataInterface);
+				const UAnimNextSharedVariables* SharedVariables = SharedVariablesEntry->GetAsset();
+				UAnimNextSharedVariables_EditorData* EditorData = UE::AnimNext::UncookedOnly::FUtils::GetEditorData<UAnimNextSharedVariables_EditorData>(SharedVariables);
 				if(InCheckForCircularity(EditorData, InCheckForCircularity))
 				{
 					return true;
@@ -1921,11 +2124,11 @@ UAnimNextDataInterfaceEntry* UAnimNextRigVMAssetEditorData::AddDataInterface(UAn
 
 	if(CheckForCircularity(EditorData, CheckForCircularity))
 	{
-		ReportError(TEXT("UAnimNextRigVMAssetEditorData::AddDataInterface: Circular reference detected."));
+		ReportError(TEXT("UAnimNextRigVMAssetEditorData::AddSharedVariables: Circular reference detected."));
 		return nullptr;
 	}
 
-	auto CheckForPublicMembers = [](UAnimNextDataInterface_EditorData* InEditorData, auto& InCheckForPublicMembers)
+	auto CheckForPublicMembers = [](UAnimNextSharedVariables_EditorData* InEditorData, auto& InCheckForPublicMembers)
 	{
 		for(UAnimNextRigVMAssetEntry* Entry : InEditorData->Entries)
 		{
@@ -1936,10 +2139,10 @@ UAnimNextDataInterfaceEntry* UAnimNextRigVMAssetEditorData::AddDataInterface(UAn
 					return true;
 				}
 			}
-			else if(UAnimNextDataInterfaceEntry* DataInterfaceEntry = Cast<UAnimNextDataInterfaceEntry>(Entry))
+			else if(UAnimNextSharedVariablesEntry* SharedVariablesEntry = Cast<UAnimNextSharedVariablesEntry>(Entry))
 			{
-				UAnimNextDataInterface* DataInterface = DataInterfaceEntry->GetDataInterface();
-				UAnimNextDataInterface_EditorData* EditorData = UE::AnimNext::UncookedOnly::FUtils::GetEditorData<UAnimNextDataInterface_EditorData>(DataInterface);
+				const UAnimNextSharedVariables* SharedVariables = SharedVariablesEntry->GetAsset();
+				UAnimNextSharedVariables_EditorData* EditorData = UE::AnimNext::UncookedOnly::FUtils::GetEditorData<UAnimNextSharedVariables_EditorData>(SharedVariables);
 				if(InCheckForPublicMembers(EditorData, InCheckForPublicMembers))
 				{
 					return true;
@@ -1952,28 +2155,28 @@ UAnimNextDataInterfaceEntry* UAnimNextRigVMAssetEditorData::AddDataInterface(UAn
 
 	if(!CheckForPublicMembers(EditorData, CheckForPublicMembers))
 	{
-		ReportError(TEXT("UAnimNextRigVMAssetEditorData::AddDataInterface: No public variables found."));
+		ReportError(TEXT("UAnimNextRigVMAssetEditorData::AddSharedVariables: No public variables found."));
 		return nullptr;
 	}
 
-	// Check for duplicate interface
-	auto DuplicatePredicate = [InDataInterface](const UAnimNextRigVMAssetEntry* InEntry)
+	// Check for duplicate entry
+	auto DuplicatePredicate = [InSharedVariables](const UAnimNextRigVMAssetEntry* InEntry)
 	{
-		if(const UAnimNextDataInterfaceEntry* InterfaceEntry = Cast<UAnimNextDataInterfaceEntry>(InEntry))
+		if(const UAnimNextSharedVariablesEntry* SharedVariablesEntry = Cast<UAnimNextSharedVariablesEntry>(InEntry))
 		{
-			return InterfaceEntry->DataInterface == InDataInterface;
+			return SharedVariablesEntry->GetAsset() == InSharedVariables;
 		}
 		return false;
 	};
 
 	if(Entries.ContainsByPredicate(DuplicatePredicate))
 	{
-		ReportError(TEXT("UAnimNextRigVMAssetEditorData::AddDataInterface: Data interface already implemented."));
+		ReportError(TEXT("UAnimNextRigVMAssetEditorData::AddSharedVariables: Shared variables already referenced."));
 		return nullptr;
 	}
 
-	UAnimNextDataInterfaceEntry* NewEntry = CreateNewSubEntry<UAnimNextDataInterfaceEntry>(this);
-	NewEntry->SetDataInterface(InDataInterface);
+	UAnimNextSharedVariablesEntry* NewEntry = CreateNewSubEntry<UAnimNextSharedVariablesEntry>(this);
+	NewEntry->SetAsset(InSharedVariables);
 	NewEntry->Initialize(this);
 
 	if(bSetupUndoRedo)
@@ -1991,14 +2194,80 @@ UAnimNextDataInterfaceEntry* UAnimNextRigVMAssetEditorData::AddDataInterface(UAn
 	if (bPrintPythonCommand)
 	{
 		RigVMPythonUtils::Print(GetName(),
-				FString::Printf(TEXT("asset.add_data_interface(unreal.find_object(outer=None, name='%s'))"),
-				 *InDataInterface->GetPathName()));
+				FString::Printf(TEXT("asset.add_shared_variables(unreal.find_object(outer=None, name='%s'))"),
+				 *InSharedVariables->GetPathName()));
 	}
-	
 
 	return NewEntry;
 }
 
+
+UAnimNextSharedVariablesEntry* UAnimNextRigVMAssetLibrary::AddSharedVariablesStruct(UAnimNextRigVMAsset* InAsset, UScriptStruct* InStruct, bool bSetupUndoRedo, bool bPrintPythonCommand)
+{
+	return UE::AnimNext::UncookedOnly::FUtils::GetEditorData(InAsset)->AddSharedVariablesStruct(InStruct, bSetupUndoRedo, bPrintPythonCommand);
+}
+
+UAnimNextSharedVariablesEntry* UAnimNextRigVMAssetEditorData::AddSharedVariablesStruct(const UScriptStruct* InStruct, bool bSetupUndoRedo, bool bPrintPythonCommand)
+{
+	if(InStruct == nullptr)
+	{
+		ReportError(TEXT("UAnimNextRigVMAssetEditorData::AddSharedVariablesStruct: Invalid struct supplied."));
+		return nullptr;
+	}
+
+	if(!GetEntryClasses().Contains(UAnimNextSharedVariablesEntry::StaticClass()) || !CanAddNewEntry(UAnimNextSharedVariablesEntry::StaticClass()))
+	{
+		ReportError(TEXT("UAnimNextRigVMAssetEditorData::AddSharedVariablesStruct: Cannot add a shared variables struct to this asset - entry is not allowed."));
+		return nullptr;
+	}
+
+	// Check for duplicate entry
+	auto DuplicatePredicate = [InStruct](const UAnimNextRigVMAssetEntry* InEntry)
+	{
+		if(const UAnimNextSharedVariablesEntry* SharedVariablesEntry = Cast<UAnimNextSharedVariablesEntry>(InEntry))
+		{
+			return SharedVariablesEntry->GetStruct() == InStruct;
+		}
+		return false;
+	};
+
+	if(Entries.ContainsByPredicate(DuplicatePredicate))
+	{
+		ReportError(TEXT("UAnimNextRigVMAssetEditorData::AddSharedVariables: Shared variables struct already referenced."));
+		return nullptr;
+	}
+
+	UAnimNextSharedVariablesEntry* NewEntry = CreateNewSubEntry<UAnimNextSharedVariablesEntry>(this);
+	NewEntry->SetStruct(InStruct);
+	NewEntry->Initialize(this);
+
+	if(bSetupUndoRedo)
+	{
+		NewEntry->Modify();
+		Modify();
+	}
+
+	AddEntryInternal(NewEntry);
+
+	CustomizeNewAssetEntry(NewEntry);
+
+	BroadcastModified(EAnimNextEditorDataNotifType::EntryAdded, NewEntry);
+
+	if (bPrintPythonCommand)
+	{
+		RigVMPythonUtils::Print(GetName(),
+				FString::Printf(TEXT("asset.add_shared_variables_struct(unreal.find_object(outer=None, name='%s'))"),
+				 *InStruct->GetPathName()));
+	}
+
+	return NewEntry;
+}
+
+URigVMLibraryNode* UAnimNextRigVMAssetLibrary::AddFunction(UAnimNextRigVMAsset* InAsset, FName InFunctionName, bool bInMutable, bool bSetupUndoRedo, bool bPrintPythonCommand)
+{
+	return UE::AnimNext::UncookedOnly::FUtils::GetEditorData(InAsset)->AddFunction(InFunctionName, bInMutable, bSetupUndoRedo, bPrintPythonCommand);
+}
+
 URigVMLibraryNode* UAnimNextRigVMAssetEditorData::AddFunction(FName InFunctionName, bool bInMutable, bool bInSetupUndoRedo, bool bPrintPythonCommand)
 {
 	URigVMController* Controller = RigVMClient.GetOrCreateController(GetLocalFunctionLibrary());
@@ -2026,11 +2295,11 @@ bool UAnimNextRigVMAssetEditorData::HasPublicVariables() const
 				return true;
 			}
 		}
-		else if(UAnimNextDataInterfaceEntry* DataInterfaceEntry = Cast<UAnimNextDataInterfaceEntry>(Entry))
+		else if(UAnimNextSharedVariablesEntry* SharedVariablesEntry = Cast<UAnimNextSharedVariablesEntry>(Entry))
 		{
-			if(DataInterfaceEntry->DataInterface)
+			if(SharedVariablesEntry->GetAsset())
 			{
-				UAnimNextDataInterface_EditorData* EditorData = UE::AnimNext::UncookedOnly::FUtils::GetEditorData<UAnimNextDataInterface_EditorData>(DataInterfaceEntry->DataInterface.Get());
+				UAnimNextSharedVariables_EditorData* EditorData = UE::AnimNext::UncookedOnly::FUtils::GetEditorData<UAnimNextSharedVariables_EditorData>(SharedVariablesEntry->GetAsset());
 				return EditorData->HasPublicVariables();
 			}
 		}
@@ -2038,42 +2307,69 @@ bool UAnimNextRigVMAssetEditorData::HasPublicVariables() const
 	return false;
 }
 
-void UAnimNextRigVMAssetEditorData::GetPublicVariables(TArray<UAnimNextVariableEntry*>& OutPublicVariables) const
+void UAnimNextRigVMAssetEditorData::GetAllVariables(TArray<FVariableInfo>& OutVariables, EVariableRecursion InRecursion, EVariableAccessFilter InAccess) const
 {
 	for(UAnimNextRigVMAssetEntry* Entry : Entries)
 	{
 		if(UAnimNextVariableEntry* VariableEntry = Cast<UAnimNextVariableEntry>(Entry))
 		{
-			if(VariableEntry->GetExportAccessSpecifier() == EAnimNextExportAccessSpecifier::Public)
+			if (InAccess == EVariableAccessFilter::All ||
+				VariableEntry->GetExportAccessSpecifier() == EAnimNextExportAccessSpecifier::Public)
 			{
-				OutPublicVariables.Add(VariableEntry);
+				const UAnimNextRigVMAsset* Asset = UE::AnimNext::UncookedOnly::FUtils::GetAsset<UAnimNextRigVMAsset>(this);
+
+				const FProperty* Property;
+				TConstArrayView<uint8> Value;
+				VariableEntry->GetDefaultValue(Property, Value);
+				
+				OutVariables.Add(
+					{
+						VariableEntry->GetExportName(),
+						VariableEntry->GetType(),
+						Asset,
+						VariableEntry->GetExportAccessSpecifier(),
+						Property,
+						Value
+					});
 			}
 		}
-		else if(UAnimNextDataInterfaceEntry* DataInterfaceEntry = Cast<UAnimNextDataInterfaceEntry>(Entry))
+		else if(UAnimNextSharedVariablesEntry* SharedVariablesEntry = Cast<UAnimNextSharedVariablesEntry>(Entry))
 		{
-			if(DataInterfaceEntry->DataInterface)
+			if (InRecursion == EVariableRecursion::SelfOnly)
 			{
-				UAnimNextDataInterface_EditorData* EditorData = UE::AnimNext::UncookedOnly::FUtils::GetEditorData<UAnimNextDataInterface_EditorData>(DataInterfaceEntry->DataInterface.Get());
-				EditorData->GetPublicVariables(OutPublicVariables);
+				continue;
 			}
-		}
-	}
-}
 
-void UAnimNextRigVMAssetEditorData::GetAllVariables(TArray<UAnimNextVariableEntry*>& OutPublicVariables) const
-{
-	for(UAnimNextRigVMAssetEntry* Entry : Entries)
-	{
-		if(UAnimNextVariableEntry* VariableEntry = Cast<UAnimNextVariableEntry>(Entry))
-		{
-			OutPublicVariables.Add(VariableEntry);
-		}
-		else if(UAnimNextDataInterfaceEntry* DataInterfaceEntry = Cast<UAnimNextDataInterfaceEntry>(Entry))
-		{
-			if(DataInterfaceEntry->DataInterface)
+			switch (SharedVariablesEntry->GetType())
 			{
-				UAnimNextDataInterface_EditorData* EditorData = UE::AnimNext::UncookedOnly::FUtils::GetEditorData<UAnimNextDataInterface_EditorData>(DataInterfaceEntry->DataInterface.Get());
-				EditorData->GetPublicVariables(OutPublicVariables);
+			case EAnimNextSharedVariablesType::Asset:
+				if(SharedVariablesEntry->GetAsset())
+				{
+					UAnimNextSharedVariables_EditorData* EditorData = UE::AnimNext::UncookedOnly::FUtils::GetEditorData<UAnimNextSharedVariables_EditorData>(SharedVariablesEntry->GetAsset());
+					EditorData->GetAllVariables(OutVariables, InRecursion, EVariableAccessFilter::PublicOnly);
+				}
+				break;
+			case EAnimNextSharedVariablesType::Struct:
+				if(SharedVariablesEntry->GetStruct())
+				{
+					for (TFieldIterator<FProperty> It(SharedVariablesEntry->GetStruct()); It; ++It)
+					{
+						const FProperty* Property = *It;
+						if (InAccess == EVariableAccessFilter::All || Property->HasAnyPropertyFlags(CPF_NativeAccessSpecifierPublic))
+						{
+							OutVariables.Add(
+								{
+									Property->GetFName(),
+									FAnimNextParamType::FromProperty(Property),
+									SharedVariablesEntry->GetStruct(),
+									EAnimNextExportAccessSpecifier::Public,
+									Property,
+									TConstArrayView<uint8>()
+								});
+						}
+					}
+				}
+				break;
 			}
 		}
 	}
@@ -2081,6 +2377,11 @@ void UAnimNextRigVMAssetEditorData::GetAllVariables(TArray<UAnimNextVariableEntr
 
 void UAnimNextRigVMAssetEditorData::HandleReportFromCompiler(EMessageSeverity::Type InSeverity, UObject* InSubject, const FString& InMessage)
 {
+	if (bSuspendCompilerReports)
+	{
+		return;
+	}
+
 	FCompilerResultsLog& Log = UE::AnimNext::UncookedOnly::FScopedCompilerResults::GetLog();
 
 	UObject* SubjectForMessage = InSubject;
@@ -2414,4 +2715,184 @@ void UAnimNextRigVMAssetEditorData::SetUseExternalPackages(TArrayView<UAnimNextR
 
 #endif
 
+FInstancedPropertyBag UAnimNextRigVMAssetEditorData::GenerateCombinedPropertyBag() const
+{
+	UAnimNextRigVMAsset* Asset = UE::AnimNext::UncookedOnly::FUtils::GetAsset<UAnimNextRigVMAsset>(this);
+
+	FInstancedPropertyBag CombinedPropertyBag;
+	auto AddInstancedPropertyBag = [this, &CombinedPropertyBag](const FInstancedPropertyBag& InPropertyBag)
+	{
+		if(const UPropertyBag* PropertyBag = InPropertyBag.GetPropertyBagStruct())
+		{
+			TConstArrayView<FPropertyBagPropertyDesc> VariableDescs = PropertyBag->GetPropertyDescs();
+			if(VariableDescs.Num() != 0)
+			{
+				CombinedPropertyBag.AddProperties(VariableDescs);
+			}
+		}
+	};
+
+	// NOTE: Order of external variables is important here!
+	// It must match that of the variables FUAFInstanceVariableData::Initialize
+
+	// First add internal variables
+	AddInstancedPropertyBag(Asset->VariableDefaults);
+
+	// Next add shared variables
+	for (const UAnimNextRigVMAsset* ReferencedVariableAsset : Asset->ReferencedVariableAssets)
+	{
+		if (ReferencedVariableAsset == nullptr)
+		{
+			continue;
+		}
+
+		AddInstancedPropertyBag(ReferencedVariableAsset->GetVariableDefaults());
+	}
+
+	// Next add native structs
+	for (const UScriptStruct* ReferencedStruct : Asset->ReferencedVariableStructs)
+	{
+		TSharedRef<UE::AnimNext::FStructDataCache> StructData = UE::AnimNext::FStructDataCache::GetStructInfo(ReferencedStruct);
+		for (const UE::AnimNext::FStructDataCache::FPropertyInfo& PropertyInfo : StructData->GetProperties())
+		{
+			CombinedPropertyBag.AddProperty(PropertyInfo.Property->GetFName(), PropertyInfo.Property);
+		}
+	}
+
+	return CombinedPropertyBag;
+}
+
+void UAnimNextRigVMAssetEditorData::UpgradeDataInterfaces()
+{
+	TGuardValue<bool> DisableEditorDataNotifications(bSuspendEditorDataNotifications, true);
+	TGuardValue<bool> DisableAutoCompile(bAutoRecompileVM, false);
+
+	PRAGMA_DISABLE_DEPRECATION_WARNINGS
+	if (NativeInterfaces_DEPRECATED.Num())
+	{
+		// Iterate over all the native interfaces and add shared variables for them
+		for (const UScriptStruct* Struct : NativeInterfaces_DEPRECATED)
+		{
+			// Translate 'native interface' into the trait data it has become
+			FInstancedStruct InstancedStruct(Struct);
+			const UScriptStruct* UpgradeStruct = InstancedStruct.Get<FAnimNextNativeDataInterface>().GetUpgradeTraitStruct();
+			if(UpgradeStruct)
+			{
+				AddSharedVariablesStruct(UpgradeStruct, false, false);
+
+				// Remove all old variables
+				for (TFieldIterator<FProperty> It(UpgradeStruct); It; ++It)
+				{
+					const FName VariableName = It->GetFName();
+
+					// Try to find a variable both with and without the 'b' prefix 
+					if (It->IsA<FBoolProperty>())
+					{
+						FString BoolName = VariableName.ToString();
+						if (BoolName.RemoveFromStart(TEXT("b")))
+						{
+							if (UAnimNextVariableEntry* VariableEntry = Cast<UAnimNextVariableEntry>(FindEntry(*BoolName)))
+							{
+								RemoveEntry(VariableEntry, false, false);
+								continue;
+							}
+						}
+					}
+
+					if (UAnimNextVariableEntry* VariableEntry = Cast<UAnimNextVariableEntry>(FindEntry(VariableName)))
+					{
+						RemoveEntry(VariableEntry, false, false);
+						continue;
+					}
+				}
+			}
+		}
+
+		NativeInterfaces_DEPRECATED.Empty();
+	}
+	PRAGMA_ENABLE_DEPRECATION_WARNINGS
+
+	{
+		// Recompile to incorporate any new native interface variables (will create errors, so suppress them)
+		TGuardValue<bool> DisableCompilerReports(bSuspendCompilerReports, true);
+		RecompileVM();
+	}
+
+	// Now replace variable nodes with scoped ones in all graphs, if possible
+	for (URigVMGraph* Graph : GetRigVMClient()->GetAllModels(true, true))
+	{
+		UAnimNextControllerBase* Controller = CastChecked<UAnimNextControllerBase>(GetController(Graph));
+		TArray<URigVMNode*> Nodes = Graph->GetNodes();
+		for (URigVMNode* Node : Nodes)
+		{
+			if (URigVMVariableNode* VariableNode = Cast<URigVMVariableNode>(Node))
+			{
+				const FName OriginalVariableName = VariableNode->GetVariableName();
+				FName NewVariableName = OriginalVariableName;
+				TArray<FName> VariableNames;
+				VariableNames.Add(OriginalVariableName);
+				if (VariableNode->GetCPPType() == TEXT("bool"))
+				{
+					TStringBuilder<128> StringBuilder;
+					StringBuilder.Append(TEXT("b"));
+					OriginalVariableName.AppendString(StringBuilder);
+					VariableNames.Add(StringBuilder.ToString());
+				}
+
+				// See if the variable exists in our shared variable entries
+				for(UAnimNextRigVMAssetEntry* Entry : Entries)
+				{
+					const UObject* AssetOrStruct = nullptr;
+					if (UAnimNextSharedVariablesEntry* SharedVariablesEntry = Cast<UAnimNextSharedVariablesEntry>(Entry))
+					{
+						for (FName VariableName : VariableNames)
+						{
+							switch (SharedVariablesEntry->GetType())
+							{
+							case EAnimNextSharedVariablesType::Asset:
+								{
+									if (const UAnimNextRigVMAsset* Asset = SharedVariablesEntry->GetAsset())
+									{
+										const FInstancedPropertyBag& Variables = Asset->GetVariableDefaults();
+										if (Variables.FindPropertyDescByName(VariableName) != nullptr)
+										{
+											AssetOrStruct = Asset;
+											NewVariableName = VariableName;
+										}
+									}
+									break;
+								}
+							case EAnimNextSharedVariablesType::Struct:
+								{
+									if (const UScriptStruct* Struct = SharedVariablesEntry->GetStruct())
+									{
+										if (Struct->FindPropertyByName(VariableName) != nullptr)
+										{
+											AssetOrStruct = Struct;
+											NewVariableName = VariableName;
+										}
+									}
+									break;
+								}
+							}
+
+							if (AssetOrStruct)
+							{
+								break;
+							}
+						}
+					}
+
+					if (AssetOrStruct)
+					{
+						Controller->ReplaceVariableNodeWithSharedVariableNode(VariableNode, NewVariableName, AssetOrStruct);
+						break;
+					}
+				}
+			}
+		}
+	}
+}
+
+
 #undef LOCTEXT_NAMESPACE
diff --git a/Engine/Plugins/Experimental/AnimNext/Source/AnimNextUncookedOnly/Private/Entries/AnimNextSharedVariablesEntry.cpp b/Engine/Plugins/Experimental/AnimNext/Source/AnimNextUncookedOnly/Private/Entries/AnimNextSharedVariablesEntry.cpp
new file mode 100644
--- /dev/null
+++ b/Engine/Plugins/Experimental/AnimNext/Source/AnimNextUncookedOnly/Private/Entries/AnimNextSharedVariablesEntry.cpp
@@ -0,0 +1,145 @@
+﻿// Copyright Epic Games, Inc. All Rights Reserved.
+
+#include "Entries/AnimNextSharedVariablesEntry.h"
+
+#if WITH_LIVE_CODING
+#include "ILiveCodingModule.h"
+#endif
+#include "UncookedOnlyUtils.h"
+#include "Variables/AnimNextSharedVariables_EditorData.h"
+#include "Entries/AnimNextVariableEntry.h"
+#include "Logging/StructuredLog.h"
+#include "Modules/ModuleManager.h"
+#include "Param/ParamType.h"
+#include "Variables/AnimNextSharedVariables.h"
+
+#define LOCTEXT_NAMESPACE "AnimNextSharedVariablesEntry"
+
+void UAnimNextSharedVariablesEntry::Initialize(UAnimNextRigVMAssetEditorData* InEditorData)
+{
+	Super::Initialize(InEditorData);
+
+	if(Asset)
+	{
+		UAnimNextSharedVariables_EditorData* EditorData = UE::AnimNext::UncookedOnly::FUtils::GetEditorData<UAnimNextSharedVariables_EditorData>(Asset.Get());
+		EditorData->ModifiedDelegate.AddUObject(this, &UAnimNextSharedVariablesEntry::HandleAssetModified);
+	}
+	else if (Struct)
+	{
+#if WITH_LIVE_CODING
+		if (ILiveCodingModule* LiveCoding = FModuleManager::LoadModulePtr<ILiveCodingModule>(LIVE_CODING_MODULE_NAME))
+		{
+			LiveCoding->GetOnPatchCompleteDelegate().AddUObject(this, &UAnimNextSharedVariablesEntry::HandlePatchComplete);
+		}
+#endif
+	}
+}
+
+FName UAnimNextSharedVariablesEntry::GetEntryName() const
+{
+	return Asset ? Asset->GetFName() : NAME_None;
+}
+
+FText UAnimNextSharedVariablesEntry::GetDisplayName() const
+{
+	switch (Type)
+	{
+	case EAnimNextSharedVariablesType::Asset:
+		return Asset ? FText::FromName(Asset->GetFName()) : LOCTEXT("InvalidSharedVariables", "Invalid Shared Variables");
+	case EAnimNextSharedVariablesType::Struct:
+		return Struct ? Struct->GetDisplayNameText() : LOCTEXT("InvalidSharedVariables", "Invalid Shared Variables");
+	}
+
+	return FText::GetEmpty();
+}
+
+FText UAnimNextSharedVariablesEntry::GetDisplayNameTooltip() const
+{
+	switch (Type)
+	{
+	case EAnimNextSharedVariablesType::Asset:
+		return Asset ? FText::FromName(Asset->GetFName()) : LOCTEXT("InvalidSharedVariablesTooltip", "Invalid or deleted Shared Variables");
+	case EAnimNextSharedVariablesType::Struct:
+		return Struct ? FText::FromString(Struct->GetPathName()) : LOCTEXT("InvalidSharedVariables", "Invalid Shared Variables");
+	}
+
+	return FText::GetEmpty();
+}
+
+void UAnimNextSharedVariablesEntry::SetAsset(const UAnimNextSharedVariables* InAsset, bool bSetupUndoRedo)
+{
+	check(InAsset != nullptr);
+
+	if(bSetupUndoRedo)
+	{
+		Modify();
+	}
+
+	Type = EAnimNextSharedVariablesType::Asset;
+	Asset = InAsset;
+	ObjectPath = FSoftObjectPath(Asset);
+	Struct = nullptr;
+}
+
+const UAnimNextSharedVariables* UAnimNextSharedVariablesEntry::GetAsset() const
+{
+	return Type == EAnimNextSharedVariablesType::Asset ? Asset : nullptr;
+}
+
+FSoftObjectPath UAnimNextSharedVariablesEntry::GetObjectPath() const
+{
+	return ObjectPath;
+}
+
+void UAnimNextSharedVariablesEntry::SetStruct(const UScriptStruct* InStruct, bool bSetupUndoRedo)
+{
+	check(InStruct != nullptr);
+
+	if(bSetupUndoRedo)
+	{
+		Modify();
+	}
+
+	Type = EAnimNextSharedVariablesType::Struct;
+	Struct = InStruct;
+	ObjectPath = FSoftObjectPath(Struct);
+	Asset = nullptr;
+}
+
+const UScriptStruct* UAnimNextSharedVariablesEntry::GetStruct() const
+{
+	return Type == EAnimNextSharedVariablesType::Struct ? Struct : nullptr;
+}
+
+void UAnimNextSharedVariablesEntry::HandleAssetModified(UAnimNextRigVMAssetEditorData* InEditorData, EAnimNextEditorDataNotifType InType, UObject* InSubject)
+{
+	switch(InType)
+	{
+	case EAnimNextEditorDataNotifType::UndoRedo:
+	case EAnimNextEditorDataNotifType::EntryAdded:
+	case EAnimNextEditorDataNotifType::EntryRemoved:
+	case EAnimNextEditorDataNotifType::EntryRenamed:
+	case EAnimNextEditorDataNotifType::EntryAccessSpecifierChanged:
+	case EAnimNextEditorDataNotifType::VariableTypeChanged:
+	case EAnimNextEditorDataNotifType::VariableDefaultValueChanged:
+		if(UAnimNextRigVMAssetEditorData* EditorData = GetTypedOuter<UAnimNextRigVMAssetEditorData>())
+		{
+			EditorData->RequestAutoVMRecompilation();
+		}
+		break;
+	default:
+		break;
+	}
+}
+
+#if WITH_LIVE_CODING
+void UAnimNextSharedVariablesEntry::HandlePatchComplete()
+{
+	if(UAnimNextRigVMAssetEditorData* EditorData = GetTypedOuter<UAnimNextRigVMAssetEditorData>())
+	{
+		EditorData->RequestAutoVMRecompilation();
+	}
+}
+#endif
+
+#undef LOCTEXT_NAMESPACE
\ No newline at end of file
diff --git a/Engine/Plugins/Experimental/AnimNext/Source/AnimNextUncookedOnly/Private/Module/AnimNextModule_EditorData.cpp b/Engine/Plugins/Experimental/AnimNext/Source/AnimNextUncookedOnly/Private/Module/AnimNextModule_EditorData.cpp
--- a/Engine/Plugins/Experimental/AnimNext/Source/AnimNextUncookedOnly/Private/Module/AnimNextModule_EditorData.cpp
+++ b/Engine/Plugins/Experimental/AnimNext/Source/AnimNextUncookedOnly/Private/Module/AnimNextModule_EditorData.cpp
@@ -8,7 +8,7 @@
 #include "UncookedOnlyUtils.h"
 #include "Module/AnimNextModule.h"
 #include "Entries/AnimNextEventGraphEntry.h"
-#include "Entries/AnimNextDataInterfaceEntry.h"
+#include "Entries/AnimNextSharedVariablesEntry.h"
 #include "Entries/AnimNextVariableEntry.h"
 #include "String/ParseTokens.h"
 #include "UObject/AssetRegistryTagsContext.h"
@@ -38,7 +38,7 @@ TConstArrayView<TSubclassOf<UAnimNextRigVMAssetEntry>> UAnimNextModule_EditorDat
 	{
 		UAnimNextEventGraphEntry::StaticClass(),
 		UAnimNextVariableEntry::StaticClass(),
-		UAnimNextDataInterfaceEntry::StaticClass(),
+		UAnimNextSharedVariablesEntry::StaticClass(),
 	};
 	
 	return Classes;
