AnimNext: Pass-by-ref runtime refactor

diff --git a/Engine/Plugins/Experimental/AnimNext/Source/AnimNextUncookedOnly/Private/UncookedOnlyUtils.cpp b/Engine/Plugins/Experimental/AnimNext/Source/AnimNextUncookedOnly/Private/UncookedOnlyUtils.cpp
--- a/Engine/Plugins/Experimental/AnimNext/Source/AnimNextUncookedOnly/Private/UncookedOnlyUtils.cpp
+++ b/Engine/Plugins/Experimental/AnimNext/Source/AnimNextUncookedOnly/Private/UncookedOnlyUtils.cpp
@@ -25,12 +25,13 @@
 #include "Variables/IVariableBindingType.h"
 #include "Variables/RigUnit_CopyModuleProxyVariables.h"
 #include "WorkspaceAssetRegistryInfo.h"
-#include "Entries/AnimNextDataInterfaceEntry.h"
-#include "DataInterface/AnimNextDataInterface.h"
-#include "DataInterface/AnimNextDataInterface_EditorData.h"
+#include "Entries/AnimNextSharedVariablesEntry.h"
+#include "Variables/AnimNextSharedVariables.h"
+#include "Variables/AnimNextSharedVariables_EditorData.h"
 #include "Logging/MessageLog.h"
 #include "RigVMFunctions/Execution/RigVMFunction_UserDefinedEvent.h"
 #include "UObject/AssetRegistryTagsContext.h"
+#include "Variables/AnimNextSoftVariableReference.h"
 
 #define LOCTEXT_NAMESPACE "AnimNextUncookedOnlyUtils"
 
@@ -56,17 +57,10 @@ void FUtils::RecreateVM(UAnimNextRigVMAsset* InAsset)
 void FUtils::CompileVariables(const FRigVMCompileSettings& InSettings, UAnimNextRigVMAsset* InAsset, FAnimNextGetVariableCompileContext& OutCompileContext)
 {
 	check(InAsset);
-
-	UAnimNextDataInterface* DataInterface = Cast<UAnimNextDataInterface>(InAsset);
-	if(DataInterface == nullptr)
-	{
-		// Currently only support data interface types (TODO: could make UAnimNextDataInterface the common base rather than UAnimNextRigVMAsset)
-		return;
-	}
-
+	
 	FMessageLog Log("AnimNextCompilerResults");
 
-	UAnimNextDataInterface_EditorData* EditorData = GetEditorData<UAnimNextDataInterface_EditorData>(DataInterface);
+	UAnimNextRigVMAssetEditorData* EditorData = GetEditorData<UAnimNextRigVMAssetEditorData>(InAsset);
 
 	// Gather programmatic variables regenerated each compile
 	EditorData->OnPreCompileGetProgrammaticVariables(InSettings, OutCompileContext);
@@ -74,232 +68,134 @@ void FUtils::CompileVariables(const FRigVMCompileSettings& InSettings, UAnimNext
 
 	struct FStructEntryInfo
 	{
-		const UAnimNextDataInterface* FromDataInterface = nullptr;
-		const UScriptStruct* NativeInterface = nullptr;
 		FName Name;
 		FAnimNextParamType Type;
 		EAnimNextExportAccessSpecifier AccessSpecifier = EAnimNextExportAccessSpecifier::Private;
-		bool bAutoBindDataInterfaceToHost = false;
 		TConstArrayView<uint8> Value;
-		EPropertyFlags PropertyFlags = CPF_Edit;
+		EPropertyFlags PropertyFlags = CPF_NativeAccessSpecifierPrivate | CPF_Edit;
 	};
 
 	// Gather all variables in this asset.
-	// Variables are harvested from the valid entries and data interfaces.
-	// Data interface harvesting is performed recursively
-	// The topmost value for a data interface 'wins' if a value is to be supplied
 	TMap<FName, int32> EntryInfoIndexMap;
 	TArray<FStructEntryInfo> StructEntryInfos;
 	StructEntryInfos.Reserve(EditorData->Entries.Num() + ProgrammaticVariables.Num());
 	int32 NumPublicVariables = 0;
 
-	auto AddVariable = [&Log, &NumPublicVariables, &StructEntryInfos, &EntryInfoIndexMap](const UAnimNextVariableEntry* InVariable, const UAnimNextDataInterfaceEntry* InFromInterfaceEntry, const UAnimNextDataInterface* InFromInterface, bool bInAutoBindInterface)
+	TSet<const UAnimNextRigVMAsset*> SharedVariableAssets;
+	TSet<const UScriptStruct*> SharedVariableStructs;
+
+	auto AddVariable = [&Log, &NumPublicVariables, &StructEntryInfos, &EntryInfoIndexMap](FName InName, const FAnimNextParamType& InType, EAnimNextExportAccessSpecifier InAccess, TConstArrayView<uint8> InValue, EPropertyFlags InFlags)
 	{
-		const FName Name = InVariable->GetExportName();
-		const FAnimNextParamType& Type = InVariable->GetExportType();
-		if(!Type.IsValid())
+		if(!InType.IsValid())
 		{
-			Log.Error(FText::Format(LOCTEXT("InvalidVariableTypeFound", "Variable '{0}' with invalid type found"), FText::FromName(Name)));
+			Log.Error(FText::Format(LOCTEXT("InvalidVariableTypeFound", "Variable '{0}' with invalid type found"), FText::FromName(InName)));
 			return;
 		}
 
-		const EAnimNextExportAccessSpecifier AccessSpecifier = InVariable->GetExportAccessSpecifier();
-
-		// Check for type conflicts
-		int32* ExistingIndexPtr = EntryInfoIndexMap.Find(Name);
-		if(ExistingIndexPtr)
-		{
-			const FStructEntryInfo& ExistingInfo = StructEntryInfos[*ExistingIndexPtr];
-			if(ExistingInfo.Type != Type)
-			{
-				Log.Error(FText::Format(LOCTEXT("ConflictingVariableTypeFound", "Variable '{0}' with conflicting type found ({1} vs {2})"), FText::FromName(Name), FText::FromString(ExistingInfo.Type.ToString()), FText::FromString(Type.ToString())));
-				return;
-			}
-
-			if(ExistingInfo.AccessSpecifier != AccessSpecifier)
-			{
-				Log.Error(FText::Format(LOCTEXT("ConflictingVariableAccessFound", "Variable '{0}' with conflicting access specifier found ({1} vs {2})"), FText::FromName(Name), FText::FromString(UEnum::GetValueAsString(ExistingInfo.AccessSpecifier)), FText::FromString(UEnum::GetValueAsString(AccessSpecifier))));
-				return;
-			}
-		}
-		else
+		// Check for conflicts
+		if(EntryInfoIndexMap.Contains(InName))
 		{
-			if(AccessSpecifier == EAnimNextExportAccessSpecifier::Public)
-			{
-				NumPublicVariables++;
-			}
+			Log.Error(FText::Format(LOCTEXT("DuplicateVariableFound", "Variable '{0}' with duplicate name found"), FText::FromName(InName)));
+			return;
 		}
 
-		// Check the overrides to see if this variable's default is overriden
-		const FProperty* OverrideProperty = nullptr;
-		TConstArrayView<uint8> OverrideValue;
-		if(InFromInterfaceEntry != nullptr)
+		if(InAccess == EAnimNextExportAccessSpecifier::Public)
 		{
-			InFromInterfaceEntry->FindValueOverrideRecursive(Name, OverrideProperty, OverrideValue);
+			NumPublicVariables++;
 		}
 
-		TConstArrayView<uint8> Value;
-		if(!OverrideValue.IsEmpty())
-		{
-			Value = OverrideValue;
-		}
-		else
-		{
-			Value = TConstArrayView<uint8>(InVariable->GetValuePtr(), Type.GetSize());
-		}
+		int32 Index = StructEntryInfos.Add(
+			{
+				InName,
+				FAnimNextParamType(InType.GetValueType(), InType.GetContainerType(), InType.GetValueTypeObject()),
+				InAccess,
+				InValue,
+				InFlags
+			});
+
+		EntryInfoIndexMap.Add(InName, Index);
+	};
 
-		if(ExistingIndexPtr)
-		{
-			// Found a variable of the same name/type, overwrite its value
-			StructEntryInfos[*ExistingIndexPtr].Value = Value;
-		}
-		else
+	for(UAnimNextRigVMAssetEntry* Entry : EditorData->Entries)
+	{
+		if(const UAnimNextVariableEntry* VariableEntry = Cast<UAnimNextVariableEntry>(Entry))
 		{
-			// This is a new variable, check if it belongs to a native interface
-			const UAnimNextDataInterface_EditorData* FromInterfaceEditorData = GetEditorData<const UAnimNextDataInterface_EditorData>(InFromInterface);
-			check(FromInterfaceEditorData != nullptr);
-
-			bool bFound = false;
-			for (const UScriptStruct* NativeInterfaceStruct : FromInterfaceEditorData->NativeInterfaces)
+			EPropertyFlags Flags = CPF_Edit;
+			if (VariableEntry->GetExportAccessSpecifier() == EAnimNextExportAccessSpecifier::Private)
 			{
-				if (NativeInterfaceStruct != nullptr && NativeInterfaceStruct->FindPropertyByName(Name) != nullptr)
-				{
-					// Found it
-					int32 Index = StructEntryInfos.Add(
-						{
-							InFromInterface,
-							NativeInterfaceStruct,
-							Name,
-							FAnimNextParamType(Type.GetValueType(), Type.GetContainerType(), Type.GetValueTypeObject()),
-							AccessSpecifier,
-							bInAutoBindInterface,
-							Value,
-							CPF_Edit
-						});
-
-					EntryInfoIndexMap.Add(Name, Index);
-					bFound = true;
-					break;
-				}
+				Flags |= CPF_NativeAccessSpecifierPrivate;
 			}
-
-			if (!bFound)
+			else
 			{
-				// Legacy codepath
-				int32 Index = StructEntryInfos.Add(
-					{
-						InFromInterface,
-						nullptr,
-						Name,
-						FAnimNextParamType(Type.GetValueType(), Type.GetContainerType(), Type.GetValueTypeObject()),
-						AccessSpecifier,
-						bInAutoBindInterface,
-						Value,
-						CPF_Edit
-					});
-
-				EntryInfoIndexMap.Add(Name, Index);
+				Flags |= CPF_NativeAccessSpecifierPublic;
 			}
+			
+			AddVariable(
+				VariableEntry->GetExportName(),
+				VariableEntry->GetType(),
+				VariableEntry->GetExportAccessSpecifier(),
+				TConstArrayView<uint8>(VariableEntry->GetValuePtr(), VariableEntry->GetType().GetSize()),
+				Flags);
 		}
-	};
-
-	auto AddDataInterface = [&Log, &AddVariable, &DataInterface](const UAnimNextDataInterface* InDataInterface, const UAnimNextDataInterfaceEntry* InDataInterfaceEntry, bool bInPublicOnly, bool bInAutoBindInterface, auto& InAddDataInterface) -> void
-	{
-		const UAnimNextDataInterface_EditorData* DataInterfaceEditorData = GetEditorData<const UAnimNextDataInterface_EditorData>(InDataInterface);
-		check(DataInterfaceEditorData != nullptr);
-
-		for(UAnimNextRigVMAssetEntry* OtherEntry : DataInterfaceEditorData->Entries)
+		else if(const UAnimNextSharedVariablesEntry* SharedVariablesEntry = Cast<UAnimNextSharedVariablesEntry>(Entry))
 		{
-			if(const UAnimNextVariableEntry* VariableEntry = Cast<UAnimNextVariableEntry>(OtherEntry))
-			{
-				if(!bInPublicOnly || VariableEntry->GetExportAccessSpecifier() == EAnimNextExportAccessSpecifier::Public)
-				{
-					AddVariable(VariableEntry, InDataInterfaceEntry, InDataInterface, bInAutoBindInterface);
-				}
-			}
-			else if(const UAnimNextDataInterfaceEntry* DataInterfaceEntry = Cast<UAnimNextDataInterfaceEntry>(OtherEntry))
+			
+			switch (SharedVariablesEntry->GetType())
 			{
-				UAnimNextDataInterface* SubDataInterface = DataInterfaceEntry->GetDataInterface();
-				if(SubDataInterface == nullptr)
+			case EAnimNextSharedVariablesType::Asset:
 				{
-					Log.Error(FText::Format(LOCTEXT("MissingDataInterfaceWarning", "Invalid data interface found: {0}"), FText::FromString(DataInterfaceEntry->GetDataInterfacePath().ToString())));
-					return;
-				}
-				else if(DataInterface == SubDataInterface)
-				{
-					Log.Error(FText::Format(LOCTEXT("CircularDataInterfaceRefError", "Circular data interface reference found: {0}"), FText::FromString(DataInterfaceEntry->GetDataInterfacePath().ToString())));
-					return;
+					const UAnimNextSharedVariables* SharedVariablesAsset = SharedVariablesEntry->GetAsset();
+					if (SharedVariablesAsset == nullptr)
+					{
+						Log.Error(FText::Format(LOCTEXT("InvalidSharedVariablesFound", "Invalid shared variables asset found ({0})"), FText::FromString(SharedVariablesEntry->GetObjectPath().ToString())));
+						continue;
+					}
+
+					if (SharedVariableAssets.Contains(SharedVariablesAsset))
+					{
+						Log.Error(FText::Format(LOCTEXT("DuplicateSharedVariablesFound", "Duplicate shared variables asset found ({0})"), FText::FromString(SharedVariablesEntry->GetObjectPath().ToString())));
+						continue;
+					}
+
+					SharedVariableAssets.Add(SharedVariablesAsset);
+					break;
 				}
-				else
+			case EAnimNextSharedVariablesType::Struct:
 				{
-					bool bAutoBindInterface = DataInterfaceEntry->AutomaticBinding == EAnimNextDataInterfaceAutomaticBindingMode::BindSharedInterfaces;
-					InAddDataInterface(SubDataInterface, DataInterfaceEntry, true, bAutoBindInterface, InAddDataInterface);
+					const UScriptStruct* SharedVariablesStruct = SharedVariablesEntry->GetStruct();
+					if (SharedVariablesStruct == nullptr)
+					{
+						Log.Error(FText::Format(LOCTEXT("InvalidSharedVariablesFound", "Invalid shared variables struct found ({0})"), FText::FromString(SharedVariablesEntry->GetObjectPath().ToString())));
+						continue;
+					}
+
+					if (SharedVariableStructs.Contains(SharedVariablesStruct))
+					{
+						Log.Error(FText::Format(LOCTEXT("DuplicateSharedVariablesFound", "Duplicate shared variables struct found ({0})"), FText::FromString(SharedVariablesEntry->GetObjectPath().ToString())));
+						continue;
+					}
+
+					SharedVariableStructs.Add(SharedVariablesStruct);
+					break;
 				}
 			}
 		}
-	};
-
-	AddDataInterface(DataInterface, nullptr, false, true, AddDataInterface);
+	}
 
-	auto AddProgrammaticVariable = [&Log, &StructEntryInfos, &EntryInfoIndexMap](const FAnimNextProgrammaticVariable& ProgrammaticVariable)
+	for (const FAnimNextProgrammaticVariable& Variable : ProgrammaticVariables)
 	{
-		const FName Name = ProgrammaticVariable.Name;
-		const FAnimNextParamType& Type = ProgrammaticVariable.Type;
-		if(!Type.IsValid())
-		{
-			Log.Error(FText::Format(LOCTEXT("InvalidProgrammaticVariableTypeFound", "Programmatic Variable '{0}' with invalid type found"), FText::FromName(Name)));
-			return;
-		}
-
-		// Check for type conflicts
-		int32* ExistingIndexPtr = EntryInfoIndexMap.Find(Name);
-		if(ExistingIndexPtr)
-		{
-			Log.Error(FText::Format(LOCTEXT("ConflictingProgrammaticVariableFound", "Programmatic Variable '{0}' already exists, should be created new each compile with no conflicts"), FText::FromName(Name)));
-			return;
-		}
-
-		TConstArrayView<uint8> Value = TConstArrayView<uint8>(ProgrammaticVariable.GetValuePtr(), Type.GetSize());
-
-		int32 Index = StructEntryInfos.Add(
-		{
-			nullptr,
-			nullptr,
-			Name,
-			FAnimNextParamType(Type.GetValueType(), Type.GetContainerType(), Type.GetValueTypeObject()),
+		AddVariable(
+			Variable.Name,
+			Variable.Type,
 			EAnimNextExportAccessSpecifier::Private,
-			false,
-			Value,
-			CPF_AdvancedDisplay
-		});
-
-		EntryInfoIndexMap.Add(Name, Index);
-	};
-
-	for (const FAnimNextProgrammaticVariable& ProgrammaticVariable : ProgrammaticVariables)
-	{
-		AddProgrammaticVariable(ProgrammaticVariable);
+			TConstArrayView<uint8>(Variable.GetValuePtr(), Variable.Type.GetSize()),
+			CPF_AdvancedDisplay | CPF_NativeAccessSpecifierPrivate);
 	}
 
-	// Sort public entries first, then by data interface & then by size, largest first, for better packing
-	static_assert(EAnimNextExportAccessSpecifier::Private < EAnimNextExportAccessSpecifier::Public, "Private must be less than Public as parameters are sorted internally according to this assumption");
+	// Sort by size, largest first, for better packing
 	StructEntryInfos.Sort([](const FStructEntryInfo& InLHS, const FStructEntryInfo& InRHS)
 	{
-		if(InLHS.AccessSpecifier != InRHS.AccessSpecifier)
-		{
-			return InLHS.AccessSpecifier > InRHS.AccessSpecifier;
-		}
-		else if(InLHS.FromDataInterface != InRHS.FromDataInterface)
-		{
-			// If we have a null (Ex: programmatic variables), compare ptrs so that those without interfaces are last
-			if (!InLHS.FromDataInterface || !InRHS.FromDataInterface)
-			{
-				return InLHS.FromDataInterface > InRHS.FromDataInterface;
-			}
-
-			return InLHS.FromDataInterface->GetFName().LexicalLess(InRHS.FromDataInterface->GetFName());
-		}
-		else if(InLHS.Type.GetSize() != InRHS.Type.GetSize())
+		if(InLHS.Type.GetSize() != InRHS.Type.GetSize())
 		{
 			return InLHS.Type.GetSize() > InRHS.Type.GetSize();
 		}
@@ -309,8 +205,25 @@ void FUtils::CompileVariables(const FRigVMCompileSettings& InSettings, UAnimNext
 		}
 	});
 
-	DataInterface->DefaultInjectionSiteIndex = INDEX_NONE;
+	// Set shared variable assets/structs before we setup variables as GetExternalVariablesImpl relies on these
+	if (SharedVariableAssets.Num() > 0)
+	{
+		InAsset->ReferencedVariableAssets = SharedVariableAssets.Array();
+	}
+	else
+	{
+		InAsset->ReferencedVariableAssets.Empty();
+	}
 
+	if (SharedVariableStructs.Num() > 0)
+	{
+		InAsset->ReferencedVariableStructs = SharedVariableStructs.Array();
+	}
+	else
+	{
+		InAsset->ReferencedVariableStructs.Empty();
+	}
+	
 	if(StructEntryInfos.Num() > 0)
 	{
 		// Build PropertyDescs and values to batch-create the property bag
@@ -319,81 +232,54 @@ void FUtils::CompileVariables(const FRigVMCompileSettings& InSettings, UAnimNext
 		TArray<TConstArrayView<uint8>> Values;
 		Values.Reserve(StructEntryInfos.Num());
 
-		DataInterface->ImplementedInterfaces.Empty();
-
-		for (TEnumerateRef<const FStructEntryInfo> StructEntryInfo : EnumerateRange(StructEntryInfos))
-		{
-			PropertyDescs.Emplace(StructEntryInfo->Name, StructEntryInfo->Type.ContainerType, StructEntryInfo->Type.ValueType, StructEntryInfo->Type.ValueTypeObject, StructEntryInfo->PropertyFlags);
-			Values.Add(StructEntryInfo->Value);
-
-			// Note: string comparison here because otherwise we would have a circular dependency
-			if( EditorData->DefaultInjectionSite == StructEntryInfo->Name &&
-				StructEntryInfo->Type.ValueTypeObject->GetPathName() == TEXT("/Script/AnimNextAnimGraph.AnimNextAnimGraph"))
-			{
-				DataInterface->DefaultInjectionSiteIndex = StructEntryInfo.GetIndex();
-			}
-			
-			if(StructEntryInfo->AccessSpecifier != EAnimNextExportAccessSpecifier::Public)
-			{
-				continue;
-			}
-
-			// Now process any data interfaces (sets of public variables) 
-			auto CheckForExistingDataInterface = [&StructEntryInfo](const FAnimNextImplementedDataInterface& InImplementedDataInterface)
-			{
-				return InImplementedDataInterface.DataInterface == StructEntryInfo->FromDataInterface;
-			};
-
-			FAnimNextImplementedDataInterface* ExistingImplementedDataInterface = DataInterface->ImplementedInterfaces.FindByPredicate(CheckForExistingDataInterface);
-			if(ExistingImplementedDataInterface == nullptr)
-			{
-				FAnimNextImplementedDataInterface& NewImplementedDataInterface = DataInterface->ImplementedInterfaces.AddDefaulted_GetRef();
-				NewImplementedDataInterface.DataInterface = StructEntryInfo->FromDataInterface;
-				NewImplementedDataInterface.NativeInterface = StructEntryInfo->NativeInterface;
-				NewImplementedDataInterface.VariableIndex = StructEntryInfo.GetIndex();
-				NewImplementedDataInterface.NumVariables = 1;
-				NewImplementedDataInterface.bAutoBindToHost = StructEntryInfo->bAutoBindDataInterfaceToHost; 
-			}
-			else
-			{
-				ExistingImplementedDataInterface->NumVariables++;
-			}
-		}
-
-		if(EditorData->DefaultInjectionSite != NAME_None && DataInterface->DefaultInjectionSiteIndex == INDEX_NONE)
+		for (const FStructEntryInfo& StructEntryInfo : StructEntryInfos)
 		{
-			Log.Error(FText::Format(LOCTEXT("MissingDefaultInjectionSiteWarning", "Could not find default injection site: {0}"), FText::FromName(EditorData->DefaultInjectionSite)));
+			PropertyDescs.Emplace(StructEntryInfo.Name, StructEntryInfo.Type.ContainerType, StructEntryInfo.Type.ValueType, StructEntryInfo.Type.ValueTypeObject, StructEntryInfo.PropertyFlags);
+			Values.Add(StructEntryInfo.Value);
 		}
 
 		// Create new property bags and migrate
-		EPropertyBagResult Result = DataInterface->VariableDefaults.ReplaceAllPropertiesAndValues(PropertyDescs, Values);
+		EPropertyBagResult Result = InAsset->VariableDefaults.ReplaceAllPropertiesAndValues(PropertyDescs, Values);
 		check(Result == EPropertyBagResult::Success);
+	}
+	else
+	{
+		InAsset->VariableDefaults.Reset();
+	}
 
-		if(NumPublicVariables > 0)
-		{
-			TConstArrayView<FPropertyBagPropertyDesc> PublicPropertyDescs(PropertyDescs.GetData(), NumPublicVariables);
-			TConstArrayView<TConstArrayView<uint8>> PublicValues(Values.GetData(), NumPublicVariables);
-			Result = DataInterface->PublicVariableDefaults.ReplaceAllPropertiesAndValues(PublicPropertyDescs, PublicValues);
-			check(Result == EPropertyBagResult::Success);
-		}
-		else
-		{
-			DataInterface->PublicVariableDefaults.Reset();
-		}
+	// Now rebuild combined the combined property bag we used for stable properties 
+	InAsset->CombinedPropertyBag = EditorData->GenerateCombinedPropertyBag();
 
-		// Rebuild external variables
-		DataInterface->VM->SetExternalVariableDefs(DataInterface->GetExternalVariablesImpl(false));
+	// Rebuild external variables
+	if (InAsset->CombinedPropertyBag.GetNumPropertiesInBag() > 0)
+	{
+		InAsset->VM->SetExternalVariableDefs(InAsset->GetExternalVariablesImpl(false));
 	}
 	else
 	{
-		DataInterface->ImplementedInterfaces.Empty();
-		DataInterface->VariableDefaults.Reset();
-		DataInterface->PublicVariableDefaults.Reset();
-		DataInterface->VM->ClearExternalVariables(DataInterface->ExtendedExecuteContext);
-		DataInterface->DefaultInjectionSiteIndex = INDEX_NONE;
+		InAsset->VM->ClearExternalVariables(InAsset->ExtendedExecuteContext);
 	}
 
+	// Validate injection site
+	InAsset->DefaultInjectionSite.Reset();
 
+	if(!EditorData->DefaultInjectionSiteReference.IsNone())
+	{
+		if (!InAsset->DefaultInjectionSite.IsValid())
+		{
+			Log.Error(FText::Format(LOCTEXT("MissingDefaultInjectionSiteError", "Could not find default injection site: {0}"), FText::FromName(EditorData->DefaultInjectionSiteReference.GetName())));
+		}
+		else
+		{
+			// Warn about using the deprecated name-based path
+			if (EditorData->DefaultInjectionSiteReference.GetObject() == nullptr)
+			{
+				Log.Warning(FText::Format(LOCTEXT("DeprecatedNamedInjectionSiteWarning", "Default injection site '{0}' uses a name-based reference. Please select a full reference to an asset's variable."), FText::FromName(EditorData->DefaultInjectionSiteReference.GetName())));
+			}
+
+			InAsset->DefaultInjectionSite = EditorData->DefaultInjectionSiteReference; 
+		}
+	}
 }
 
 void FUtils::CompileVariableBindings(const FRigVMCompileSettings& InSettings, UAnimNextRigVMAsset* InAsset, TArray<URigVMGraph*>& OutGraphs)
@@ -814,18 +700,79 @@ void FUtils::SetupEventGraph(URigVMController* InController, UScriptStruct* InEv
 	}
 }
 
-FAnimNextParamType FUtils::GetParameterTypeFromName(FName InName)
+FAnimNextParamType FUtils::FindVariableType(const FAnimNextVariableReference& InVariableReference)
 {
-	// Query the asset registry for other params
-	TMap<FAssetData, FAnimNextAssetRegistryExports> ExportMap;
-	GetExportedVariablesFromAssetRegistry(ExportMap);
-	for(const TPair<FAssetData, FAnimNextAssetRegistryExports>& ExportPair : ExportMap)
+	// If this is likely a struct variable, just resolve & search for the member
+	if (const UScriptStruct* Struct = Cast<UScriptStruct>(InVariableReference.GetObject()))
 	{
-		for(const FAnimNextAssetRegistryExportedVariable& Parameter : ExportPair.Value.Variables)
+		if (const FProperty* Property = Struct->FindPropertyByName(InVariableReference.GetName()))
 		{
-			if(Parameter.Name == InName)
+			return FAnimNextParamType::FromProperty(Property);
+		}
+	}
+	else if (const UAnimNextRigVMAsset* Asset =  Cast<UAnimNextRigVMAsset>(InVariableReference.GetObject()))
+	{
+		const UAnimNextRigVMAssetEditorData* EditorData = GetEditorData<UAnimNextRigVMAssetEditorData>(Asset);
+		FAnimNextAssetRegistryExports Exports;
+		GetAssetVariables(EditorData, Exports);
+		for(const FAnimNextAssetRegistryExportedVariable& Variable : Exports.Variables)
+		{
+			if(Variable.Name == InVariableReference.GetName())
 			{
-				return Parameter.Type;
+				return Variable.Type;
+			}
+		}
+	}
+
+	return FAnimNextParamType();
+}
+
+
+FAnimNextParamType FUtils::FindVariableType(const FAnimNextSoftVariableReference& InSoftVariableReference)
+{
+	// If this is likely a struct variable, just resolve & search for the member
+	if (InSoftVariableReference.GetSoftObjectPath().GetLongPackageName().StartsWith(TEXT("/Script/")))
+	{
+		if (const UScriptStruct* Struct = Cast<UScriptStruct>(InSoftVariableReference.GetSoftObjectPath().TryLoad()))
+		{
+			if (const FProperty* Property = Struct->FindPropertyByName(InSoftVariableReference.GetName()))
+			{
+				return FAnimNextParamType::FromProperty(Property);
+			}
+		}
+	}
+	else
+	{
+		// Query the asset registry for the asset
+		FAssetData AssetData = IAssetRegistry::GetChecked().GetAssetByObjectPath(InSoftVariableReference.GetSoftObjectPath());
+
+		// If the asset is loaded, just use it
+		if (const UAnimNextRigVMAsset* Asset = Cast<UAnimNextRigVMAsset>(AssetData.FastGetAsset(false)))
+		{
+			const UAnimNextRigVMAssetEditorData* EditorData = GetEditorData<UAnimNextRigVMAssetEditorData>(Asset);
+			FAnimNextAssetRegistryExports Exports;
+			GetAssetVariables(EditorData, Exports);
+			for(const FAnimNextAssetRegistryExportedVariable& Variable : Exports.Variables)
+			{
+				if(Variable.Name == InSoftVariableReference.GetName())
+				{
+					return Variable.Type;
+				}
+			}
+		}
+		else
+		{
+			// Otherwise use AR data to find its exports
+			FAnimNextAssetRegistryExports Exports;
+			if (GetExportedVariablesForAsset(AssetData, Exports))
+			{
+				for(const FAnimNextAssetRegistryExportedVariable& Variable : Exports.Variables)
+				{
+					if(Variable.Name == InSoftVariableReference.GetName())
+					{
+						return Variable.Type;
+					}
+				}
 			}
 		}
 	}
@@ -910,33 +857,7 @@ bool FUtils::GetExportedFunctionsFromAssetRegistry(FName Tag, TMap<FAssetData, F
 	return OutExports.Num() > 0;
 }
 
-static void AddParamToSet(const FAnimNextAssetRegistryExportedVariable& InNewParam, TSet<FAnimNextAssetRegistryExportedVariable>& OutExports)
-{
-	if(FAnimNextAssetRegistryExportedVariable* ExistingEntry = OutExports.Find(InNewParam))
-	{
-		if(ExistingEntry->Type != InNewParam.Type)
-		{
-			UE_LOGFMT(LogAnimation, Warning, "Type mismatch between parameter {ParameterName}. {ParamType1} vs {ParamType1}", InNewParam.Name, InNewParam.Type.ToString(), ExistingEntry->Type.ToString());
-		}
-		ExistingEntry->Flags |= InNewParam.Flags;
-	}
-	else
-	{
-		OutExports.Add(InNewParam);
-	}
-}
-
-void FUtils::GetAssetVariables(const UAnimNextRigVMAssetEditorData* EditorData, FAnimNextAssetRegistryExports& OutExports)
-{
-	OutExports.Variables.Reset();
-	OutExports.Variables.Reserve(EditorData->Entries.Num());
-
-	TSet<FAnimNextAssetRegistryExportedVariable> ExportSet;
-	GetAssetVariables(EditorData, ExportSet);
-	OutExports.Variables = ExportSet.Array();
-}
-
-void FUtils::GetAssetVariables(const UAnimNextRigVMAssetEditorData* InEditorData, TSet<FAnimNextAssetRegistryExportedVariable>& OutExports)
+void FUtils::GetAssetVariables(const UAnimNextRigVMAssetEditorData* InEditorData, FAnimNextAssetRegistryExports& OutExports)
 {
 	for(const UAnimNextRigVMAssetEntry* Entry : InEditorData->Entries)
 	{
@@ -946,16 +867,19 @@ void FUtils::GetAssetVariables(const UAnimNextRigVMAssetEditorData* InEditorData
 			if(ExportInterface->GetExportAccessSpecifier() == EAnimNextExportAccessSpecifier::Public)
 			{
 				Flags |= EAnimNextExportedVariableFlags::Public;
-				FAnimNextAssetRegistryExportedVariable NewParam(ExportInterface->GetExportName(), ExportInterface->GetExportType(), Flags);
-				AddParamToSet(NewParam, OutExports);
+				OutExports.Variables.Emplace(ExportInterface->GetExportName(), ExportInterface->GetExportType(), Flags);
 			}
 		}
-		else if(const UAnimNextDataInterfaceEntry* DataInterfaceEntry = Cast<UAnimNextDataInterfaceEntry>(Entry))
+		else if(const UAnimNextSharedVariablesEntry* DataInterfaceEntry = Cast<UAnimNextSharedVariablesEntry>(Entry))
 		{
-			if(DataInterfaceEntry->DataInterface)
+			switch(DataInterfaceEntry->Type)
 			{
-				UAnimNextDataInterface_EditorData* EditorData = GetEditorData<UAnimNextDataInterface_EditorData>(DataInterfaceEntry->DataInterface.Get());
-				GetAssetVariables(EditorData, OutExports);
+			case EAnimNextSharedVariablesType::Asset:
+				OutExports.SharedVariableAssets.Emplace(DataInterfaceEntry->GetAsset());
+				break;
+			case EAnimNextSharedVariablesType::Struct:
+				OutExports.SharedVariableStructs.Emplace(DataInterfaceEntry->GetStruct());
+				break;
 			}
 		}
 	}
diff --git a/Engine/Plugins/Experimental/AnimNext/Source/AnimNextUncookedOnly/Private/Variables/AnimNextSharedVariables_EditorData.cpp b/Engine/Plugins/Experimental/AnimNext/Source/AnimNextUncookedOnly/Private/Variables/AnimNextSharedVariables_EditorData.cpp
new file mode 100644
--- /dev/null
+++ b/Engine/Plugins/Experimental/AnimNext/Source/AnimNextUncookedOnly/Private/Variables/AnimNextSharedVariables_EditorData.cpp
@@ -0,0 +1,35 @@
+﻿// Copyright Epic Games, Inc. All Rights Reserved.
+
+#include "Variables/AnimNextSharedVariables_EditorData.h"
+
+#include "Entries/AnimNextSharedVariablesEntry.h"
+#include "Entries/AnimNextVariableEntry.h"
+
+TConstArrayView<TSubclassOf<UAnimNextRigVMAssetEntry>> UAnimNextSharedVariables_EditorData::GetEntryClasses() const
+{
+	static const TSubclassOf<UAnimNextRigVMAssetEntry> Classes[] =
+	{
+		UAnimNextVariableEntry::StaticClass(),
+		UAnimNextSharedVariablesEntry::StaticClass(),
+	};
+
+	return Classes;
+}
+
+void UAnimNextSharedVariables_EditorData::CustomizeNewAssetEntry(UAnimNextRigVMAssetEntry* InNewEntry) const
+{
+	UAnimNextVariableEntry* VariableEntry = Cast<UAnimNextVariableEntry>(InNewEntry);
+	if(VariableEntry == nullptr)
+	{
+		return;
+	}
+	
+	const bool bIsSharedVariables = ExactCast<UAnimNextSharedVariables_EditorData>(this) != nullptr;
+	if(!bIsSharedVariables)
+	{
+		return;
+	}
+
+	// Force all variables in 'pure' shared variables assets to be public
+	VariableEntry->SetExportAccessSpecifier(EAnimNextExportAccessSpecifier::Public, false);
+}
\ No newline at end of file
diff --git a/Engine/Plugins/Experimental/AnimNextAnimGraph/Source/AnimNextAnimGraph/Internal/AnimNextAnimGraphSettings.h b/Engine/Plugins/Experimental/AnimNextAnimGraph/Source/AnimNextAnimGraph/Internal/AnimNextAnimGraphSettings.h
--- a/Engine/Plugins/Experimental/AnimNextAnimGraph/Source/AnimNextAnimGraph/Internal/AnimNextAnimGraphSettings.h
+++ b/Engine/Plugins/Experimental/AnimNextAnimGraph/Source/AnimNextAnimGraph/Internal/AnimNextAnimGraphSettings.h
@@ -3,9 +3,6 @@
 #pragma once
 
 #include "Graph/AnimNextAnimationGraph.h"
-#include "StructUtils/StructView.h"
-#include "Templates/SubclassOf.h"
-#include "UObject/ObjectKey.h"
 
 #include "AnimNextAnimGraphSettings.generated.h"
 
@@ -14,37 +11,7 @@
 struct FInstancedPropertyBag;
 class UAnimNextAnimationGraph;
 class UAnimNextAnimGraphSettings;
-struct FAnimNextDataInterfacePayload;
-
-USTRUCT()
-struct FAnimNextAssetGraphMapping
-{
-	GENERATED_BODY()
-
-private:
-	friend UAnimNextAnimGraphSettings;
-
-	// The object type/class that this mapping handles
-	UPROPERTY(EditAnywhere, Category = "Mapping")
-	TSoftClassPtr<UObject> AssetType;
-
-	// The animation graph class that the type maps to
-	UPROPERTY(EditAnywhere, Category = "Mapping")
-	TSoftObjectPtr<const UAnimNextAnimationGraph> AnimationGraph;
-
-	// The public variable that will be set when creating an instance of AnimationGraph, from an asset of AssetType.
-	// For example, when mapping from an AnimSequence asset to a graph, the graph will have the AnimSequence asset set on a property of this name.
-	// Leave as 'None' to not set the asset into a variable of the graph.
-	// Must be an 'object' property of a compatible type to AssetType.
-	UPROPERTY(EditAnywhere, Category = "Mapping")
-	FName Variable;
-
-	UPROPERTY(EditAnywhere, Category = "Mapping")
-	TArray<TSoftObjectPtr<const UScriptStruct>> RequiredDataInterfaces;
-
-	// Index of the next mapping for this AssetType
-	int32 NextMappingIndex = INDEX_NONE;
-};
+struct FAnimNextFactoryParams;
 
 UCLASS(MinimalAPI, Config=AnimNextAnimGraph, defaultconfig)
 class UAnimNextAnimGraphSettings : public UObject
@@ -52,72 +19,9 @@ class UAnimNextAnimGraphSettings : public UObject
 	GENERATED_BODY()
 
 public:
-	// Given an object, return an animation graph to instantiate and a set of native interface structs that can be used to communicate with the instance
-	UE_API const UAnimNextAnimationGraph* GetGraphFromObject(const UObject* InObject) const;
-
-	// Given an object, return an animation graph to instantiate and a set of native interface structs that can be used to communicate with the instance
-	UE_API const UAnimNextAnimationGraph* GetGraphFromObject(const UObject* InObject, const FAnimNextDataInterfacePayload& InGraphPayload) const;
-
-	// Given an asset class, return whether an animation graph can be made via GetGraphFromObject.
-	UE_API bool CanGetGraphFromAssetClass(const UClass* InClass) const;
-
-	// Given an object, return the name of the variable to inject into the graph's payload when making a payload for that object's graph (the
-	// 'Variable' property of its mapping.
-	UE_API FName GetInjectedVariableNameFromObject(const UObject* InObject) const;
-
-	// Given an object and the resulting graph, generate native interface payloads for the graph.
-	// If the native interface already exists in InOutGraphPayload, it will not be created by this call.
-	UE_API void GetNativePayloadFromGraph(const UObject* InObject, const UAnimNextAnimationGraph* InAnimationGraph, FAnimNextDataInterfacePayload& InOutGraphPayload) const;
-	UE_API void GetNativePayloadFromGraph(const UObject* InObject, const UAnimNextAnimationGraph* InAnimationGraph, TArray<FInstancedStruct>& InOutGraphPayload) const;
-
-	// Given an object and the resulting graph, generate an interface payload for the graph, regardless if a native payload exists for the graph.
-	// The generated instanced property bag will encompass public variables of ALL data interfaces implemented by the graph in one bundle
-	UE_API void GetNonNativePayloadFromGraph(const UObject* InObject, const UAnimNextAnimationGraph* InAnimationGraph, FAnimNextDataInterfacePayload& InOutGraphPayload) const;
-	UE_API void GetNonNativePayloadFromGraph(const UObject* InObject, const UAnimNextAnimationGraph* InAnimationGraph, FInstancedPropertyBag& InOutGraphPayload) const;
-
-	// Rebuild mappings (loading assets synchronously if required) for lookup
-	UE_API void LoadAndRebuildMappings(bool bLoadAsync);
-
-	// Get the animation graph we run by default when hosting in a module
-	const UAnimNextAnimationGraph* GetDefaultRunGraphHost() const
-	{
-		return DefaultRunGraphHost.Get();
-	}
-
 	// Gets all allowed asset classes that users can reference that map via GetGraphFromObject
 	UFUNCTION()
 	static UE_API TArray<UClass*> GetAllowedAssetClasses();
-
-private:
-#if WITH_EDITOR
-	// UObject interface
-	UE_API virtual void PostEditChangeProperty(FPropertyChangedEvent& PropertyChangedEvent) override;
-#endif
-
-	UE_API TArray<UClass*> GetAllowedAssetClassesImpl() const;
-
-	UE_API void OnDefaultRunGraphHostLoaded(const UAnimNextAnimationGraph* AnimationGraph, bool bLoadAsync);
-	UE_API void OnMappingAnimationGraphLoaded(const UAnimNextAnimationGraph* AnimationGraph, FAnimNextAssetGraphMapping& Mapping, int32 MappingIndex);
-	UE_API void FinalizeAsyncLoad();
-	
-private:
-	// The animation graph we run by default when hosting in a module
-	UPROPERTY(Config, EditAnywhere, Category = "Graphs")
-	TSoftObjectPtr<const UAnimNextAnimationGraph> DefaultRunGraphHost;
-
-	// Mappings from assets to animation graphs
-	UPROPERTY(Config, EditAnywhere, DisplayName = "Asset/Graph Mappings", Category = "Mappings")
-	TArray<FAnimNextAssetGraphMapping> AssetGraphMappings;
-
-	// Map derived from AssetGraphMappings
-	// Maps an asset type (e.g. anim sequence) to the index in the mappings array for the first entry for that object
-	TMap<TObjectKey<const UClass>, int32> AssetGraphMap;
-
-	// Counter to keep track of how mappings have finished loading
-	int32 NumMappingsLoaded = 0;
-
-	// Flag to verify correct mapping rebuild behavior during startup
-	std::atomic<bool> bMappingsBuiltAtLeastOnce = false;
 };
 
 #undef UE_API
diff --git a/Engine/Plugins/Experimental/AnimNextAnimGraph/Source/AnimNextAnimGraph/Internal/Factory/AnimGraphBuilderContext.h b/Engine/Plugins/Experimental/AnimNextAnimGraph/Source/AnimNextAnimGraph/Internal/Factory/AnimGraphBuilderContext.h
new file mode 100644
--- /dev/null
+++ b/Engine/Plugins/Experimental/AnimNextAnimGraph/Source/AnimNextAnimGraph/Internal/Factory/AnimGraphBuilderContext.h
@@ -0,0 +1,40 @@
+﻿// Copyright Epic Games, Inc. All Rights Reserved.
+
+#pragma once
+
+#include "CoreTypes.h"
+#include "TraitCore/EntryPointHandle.h"
+#include "TraitCore/TraitWriter.h"
+
+class UAnimNextAnimationGraph;
+class FAnimationAnimNextRuntimeTest_SimpleGraphBuilder;
+
+namespace UE::AnimNext
+{
+	struct FAnimGraphFactory;
+}
+
+namespace UE::AnimNext
+{
+
+struct FAnimGraphBuilderContext
+{
+private:
+	friend FAnimGraphFactory;
+	friend ::FAnimationAnimNextRuntimeTest_SimpleGraphBuilder;
+
+	// Builds a graph using the populated TraitWriter and VariableStructs
+	ANIMNEXTANIMGRAPH_API const UAnimNextAnimationGraph* Build();
+
+public:
+	// Root trait handle
+	FAnimNextEntryPointHandle RootTraitHandle;
+
+	// Trait writer used to build the graph
+	FTraitWriter TraitWriter;
+
+	// Variable structs used to communicate with the graph
+	TArray<const UScriptStruct*> VariableStructs;
+};
+
+}
\ No newline at end of file
diff --git a/Engine/Plugins/Experimental/AnimNextAnimGraph/Source/AnimNextAnimGraph/Internal/Graph/AnimNextGraphContextData.h b/Engine/Plugins/Experimental/AnimNextAnimGraph/Source/AnimNextAnimGraph/Internal/Graph/AnimNextGraphContextData.h
--- a/Engine/Plugins/Experimental/AnimNextAnimGraph/Source/AnimNextAnimGraph/Internal/Graph/AnimNextGraphContextData.h
+++ b/Engine/Plugins/Experimental/AnimNextAnimGraph/Source/AnimNextAnimGraph/Internal/Graph/AnimNextGraphContextData.h
@@ -24,7 +24,7 @@ struct FAnimNextGraphContextData : public FAnimNextModuleContextData
 	{
 	}
 
-	const FAnimNextGraphInstance& GetGraphInstance() const { return static_cast<const FAnimNextGraphInstance&>(GetDataInterfaceInstance()); }
+	const FAnimNextGraphInstance& GetGraphInstance() const { return static_cast<const FAnimNextGraphInstance&>(GetInstance()); }
 
 private:
 	friend struct FAnimNextGraphInstance;
diff --git a/Engine/Plugins/Experimental/AnimNextAnimGraph/Source/AnimNextAnimGraph/Internal/Graph/AnimNextModuleAnimGraphComponent.h b/Engine/Plugins/Experimental/AnimNextAnimGraph/Source/AnimNextAnimGraph/Internal/Graph/AnimNextModuleAnimGraphComponent.h
--- a/Engine/Plugins/Experimental/AnimNextAnimGraph/Source/AnimNextAnimGraph/Internal/Graph/AnimNextModuleAnimGraphComponent.h
+++ b/Engine/Plugins/Experimental/AnimNextAnimGraph/Source/AnimNextAnimGraph/Internal/Graph/AnimNextModuleAnimGraphComponent.h
@@ -2,22 +2,29 @@
 
 #pragma once
 
+#include "UAFAssetInstance.h"
 #include "Injection/InjectionInfo.h"
-#include "Module/AnimNextModuleInstanceComponent.h"
+#include "Module/UAFModuleInstanceComponent.h"
 #include "AnimNextModuleAnimGraphComponent.generated.h"
 
 struct FAnimNextGraphInstance;
 class UAnimNextAnimationGraph;
 
+namespace UE::AnimNext
+{
+	struct FVariableOverrides;
+	struct FGraphAllocationParams;
+}
+
 // Module component that owns/allocates/release all animation graph instances on a module
 USTRUCT()
-struct FAnimNextModuleAnimGraphComponent : public FAnimNextModuleInstanceComponent
+struct FAnimNextModuleAnimGraphComponent : public FUAFModuleInstanceComponent
 {
 	GENERATED_BODY()
 
 	FAnimNextModuleAnimGraphComponent() = default;
 
-	TWeakPtr<FAnimNextGraphInstance> AllocateInstance(const UAnimNextAnimationGraph* InAnimationGraph, FAnimNextGraphInstance* InParentInstance = nullptr, FName InEntryPoint = NAME_None);
+	TWeakPtr<FAnimNextGraphInstance> AllocateInstance(const UAnimNextAnimationGraph* InAnimationGraph, const TSharedPtr<const UE::AnimNext::FVariableOverridesCollection>& InOverrides);
 
 	void ReleaseInstance(TWeakPtr<FAnimNextGraphInstance> InInstance);
 
diff --git a/Engine/Plugins/Experimental/AnimNextAnimGraph/Source/AnimNextAnimGraph/Internal/Graph/RigUnit_AnimNextRunAnimationGraph_v1.h b/Engine/Plugins/Experimental/AnimNextAnimGraph/Source/AnimNextAnimGraph/Internal/Graph/RigUnit_AnimNextRunAnimationGraph_v1.h
--- a/Engine/Plugins/Experimental/AnimNextAnimGraph/Source/AnimNextAnimGraph/Internal/Graph/RigUnit_AnimNextRunAnimationGraph_v1.h
+++ b/Engine/Plugins/Experimental/AnimNextAnimGraph/Source/AnimNextAnimGraph/Internal/Graph/RigUnit_AnimNextRunAnimationGraph_v1.h
@@ -6,6 +6,7 @@
 #include "Graph/AnimNextAnimGraph.h"
 #include "Graph/AnimNext_LODPose.h"
 #include "Graph/RigUnit_AnimNextBase.h"
+#include "Variables/AnimNextVariableOverridesCollection.h"
 
 #include "RigUnit_AnimNextRunAnimationGraph_v1.generated.h"
 
@@ -45,6 +46,10 @@ struct FRigUnit_AnimNextRunAnimationGraph_v1 : public FRigUnit_AnimNextBase
 	UPROPERTY(EditAnywhere, Category = "Graph", meta = (Input))
 	FAnimNextGraphReferencePose ReferencePose;
 
+	// Variable overrides to be applied to subgraphs
+	UPROPERTY(EditAnywhere, Category = "Graph", meta = (Input))
+	FAnimNextVariableOverridesCollection Overrides;
+
 	// Pose result
 	UPROPERTY(EditAnywhere, Category = "Graph", meta = (Output))
 	FAnimNextGraphLODPose Result;
diff --git a/Engine/Plugins/Experimental/AnimNextAnimGraph/Source/AnimNextAnimGraph/Internal/Graph/RigUnit_AnimNextRunAnimationGraph_v2.h b/Engine/Plugins/Experimental/AnimNextAnimGraph/Source/AnimNextAnimGraph/Internal/Graph/RigUnit_AnimNextRunAnimationGraph_v2.h
--- a/Engine/Plugins/Experimental/AnimNextAnimGraph/Source/AnimNextAnimGraph/Internal/Graph/RigUnit_AnimNextRunAnimationGraph_v2.h
+++ b/Engine/Plugins/Experimental/AnimNextAnimGraph/Source/AnimNextAnimGraph/Internal/Graph/RigUnit_AnimNextRunAnimationGraph_v2.h
@@ -6,6 +6,7 @@
 #include "Graph/AnimNextAnimGraph.h"
 #include "Graph/AnimNext_LODPose.h"
 #include "Graph/RigUnit_AnimNextBase.h"
+#include "Variables/AnimNextVariableOverridesCollection.h"
 
 #include "RigUnit_AnimNextRunAnimationGraph_v2.generated.h"
 
@@ -21,7 +22,7 @@ struct FAnimNextRunGraphWorkData_v2
 
 	// Graph we are injecting into the host
 	UPROPERTY(transient)
-	FAnimNextAnimGraph InjectedGraph;
+	FAnimNextVariableReference InjectedGraphReference;
 };
 
 /** Runs an animation graph */
@@ -47,6 +48,10 @@ struct FRigUnit_AnimNextRunAnimationGraph_v2 : public FRigUnit_AnimNextBase
 	UPROPERTY(EditAnywhere, Category = "Run", meta = (Input))
 	FAnimNextGraphReferencePose ReferencePose;
 
+	// Variable overrides to be applied to subgraphs
+	UPROPERTY(EditAnywhere, Category = "Run", meta = (Input))
+	FAnimNextVariableOverridesCollection Overrides;
+
 	// Pose result
 	UPROPERTY(EditAnywhere, Category = "Run", meta = (Output))
 	FAnimNextGraphLODPose Result;
diff --git a/Engine/Plugins/Experimental/AnimNextAnimGraph/Source/AnimNextAnimGraph/Internal/Injection/InjectionCallbackProxy.h b/Engine/Plugins/Experimental/AnimNextAnimGraph/Source/AnimNextAnimGraph/Internal/Injection/InjectionCallbackProxy.h
--- a/Engine/Plugins/Experimental/AnimNextAnimGraph/Source/AnimNextAnimGraph/Internal/Injection/InjectionCallbackProxy.h
+++ b/Engine/Plugins/Experimental/AnimNextAnimGraph/Source/AnimNextAnimGraph/Internal/Injection/InjectionCallbackProxy.h
@@ -41,10 +41,9 @@ public:
 	UFUNCTION(BlueprintCallable, meta = (BlueprintInternalUseOnly = "true"))
 	static ANIMNEXTANIMGRAPH_API UInjectionCallbackProxy* CreateProxyObjectForInjection(
 		UAnimNextComponent* AnimNextComponent,
-		FName SiteName,
+		UPARAM(meta=(AllowedType=FAnimNextAnimGraph)) FAnimNextVariableReference Site,
 		UObject* Object,
-		UAnimNextComponent* BindingComponent,
-		FInstancedStruct Payload,
+		FAnimNextFactoryParams FactoryParams = FAnimNextFactoryParams(),
 		FAnimNextInjectionBlendSettings BlendInSettings = FAnimNextInjectionBlendSettings(),
 		FAnimNextInjectionBlendSettings BlendOutSettings = FAnimNextInjectionBlendSettings());
 
@@ -52,6 +51,12 @@ public:
 	UFUNCTION(BlueprintCallable, Category="Animation|AnimNext", meta=(ExpandEnumAsExecs = "ReturnValue"))
 	ANIMNEXTANIMGRAPH_API EUninjectionResult Uninject();
 
+	// Sets a variable's value.
+	// @param    Name     The variable to set
+	// @param    Value    The value to set the variable to
+	UFUNCTION(BlueprintCallable, BlueprintInternalUseOnly, Category = "AnimNext", DisplayName = "Set Variable", CustomThunk, meta = (CustomStructureParam = Value, AutoCreateRefTerm = "Variable, Value", UnsafeDuringActorConstruction))
+	ANIMNEXTANIMGRAPH_API void SetVariable(const FAnimNextVariableReference& Variable, const int32& Value);
+
 	// UObject Interface
 	ANIMNEXTANIMGRAPH_API virtual void BeginDestroy() override;
 
@@ -66,16 +71,17 @@ protected:
 	// Attempts to play an object with the specified payload. Returns whether it started or not.
 	ANIMNEXTANIMGRAPH_API bool Inject(
 		UAnimNextComponent* AnimNextComponent,
-		FName SiteName,
+		FAnimNextVariableReference Site,
 		UObject* Object,
-		UAnimNextComponent* BindingComponent,
-		FInstancedStruct&& Payload,
+		FAnimNextFactoryParams&& FactoryParams,
 		const UE::AnimNext::FInjectionBlendSettings& BlendInSettings,
 		const UE::AnimNext::FInjectionBlendSettings& BlendOutSettings);
 
 private:
 	void Reset();
 
+	DECLARE_FUNCTION(execSetVariable);
+
 	UE::AnimNext::FInjectionRequestPtr PlayingRequest;
 	bool bWasInterrupted = false;
 };
diff --git a/Engine/Plugins/Experimental/AnimNextAnimGraph/Source/AnimNextAnimGraph/Internal/TraitCore/NodeTemplateBuilder.h b/Engine/Plugins/Experimental/AnimNextAnimGraph/Source/AnimNextAnimGraph/Internal/TraitCore/NodeTemplateBuilder.h
--- a/Engine/Plugins/Experimental/AnimNextAnimGraph/Source/AnimNextAnimGraph/Internal/TraitCore/NodeTemplateBuilder.h
+++ b/Engine/Plugins/Experimental/AnimNextAnimGraph/Source/AnimNextAnimGraph/Internal/TraitCore/NodeTemplateBuilder.h
@@ -29,15 +29,15 @@ namespace UE::AnimNext
 
 		// Returns a node template for the provided list of traits
 		// The node template will be built into the provided buffer and a pointer to it is returned
-		static UE_API FNodeTemplate* BuildNodeTemplate(const TArray<FTraitUID>& InTraitUIDs, TArray<uint8>& NodeTemplateBuffer);
+		static UE_API FNodeTemplate* BuildNodeTemplate(TConstArrayView<FTraitUID> InTraitUIDs, TArray<uint8>& NodeTemplateBuffer);
 
 		// Resets the node template builder
 		UE_API void Reset();
 
 	private:
-		static UE_API uint32 GetNodeTemplateUID(const TArray<FTraitUID>& InTraitUIDs);
+		static UE_API uint32 GetNodeTemplateUID(TConstArrayView<FTraitUID> InTraitUIDs);
 		static UE_API void AppendTemplateTrait(
-			const TArray<FTraitUID>& InTraitUIDs, int32 TraitIndex,
+			TConstArrayView<FTraitUID>, int32 TraitIndex,
 			TArray<uint8>& NodeTemplateBuffer);
 
 		TArray<FTraitUID> TraitUIDs;	// The list of traits to use when building the node template
diff --git a/Engine/Plugins/Experimental/AnimNextAnimGraph/Source/AnimNextAnimGraph/Internal/TraitCore/TraitWriter.h b/Engine/Plugins/Experimental/AnimNextAnimGraph/Source/AnimNextAnimGraph/Internal/TraitCore/TraitWriter.h
--- a/Engine/Plugins/Experimental/AnimNextAnimGraph/Source/AnimNextAnimGraph/Internal/TraitCore/TraitWriter.h
+++ b/Engine/Plugins/Experimental/AnimNextAnimGraph/Source/AnimNextAnimGraph/Internal/TraitCore/TraitWriter.h
@@ -4,14 +4,16 @@
 
 #include "CoreMinimal.h"
 
-#if WITH_EDITOR
 #include "Serialization/MemoryWriter.h"
+#include "StructUtils/StructView.h"
 
 #include "TraitCore/LatentPropertyHandle.h"
 #include "TraitCore/NodeHandle.h"
 #include "TraitCore/NodeTemplateRegistryHandle.h"
 #include "UObject/SoftObjectPath.h"
 
+struct FAnimNextTraitSharedData;
+
 namespace UE::AnimNext
 {
 	struct FNodeTemplate;
@@ -47,6 +49,7 @@ namespace UE::AnimNext
 		// Called once node writing has terminated
 		ANIMNEXTANIMGRAPH_API void EndNodeWriting();
 
+#if WITH_EDITOR
 		// Writes out the provided node using the trait properties
 		// Nodes must be written in the same order they were registered in
 		ANIMNEXTANIMGRAPH_API void WriteNode(
@@ -54,6 +57,14 @@ namespace UE::AnimNext
 			const TFunction<FString(uint32 TraitIndex, FName PropertyName)>& GetTraitProperty,
 			const TFunction<uint16(uint32 TraitIndex, FName PropertyName)>& GetTraitLatentPropertyIndex
 			);
+#endif
+
+		// Writes out the provided node using the trait properties
+		// Nodes must be written in the same order they were registered in
+		void WriteNode(
+			const FNodeHandle NodeHandle,
+			TFunctionRef<uint16(uint32 TraitIndex, FName PropertyName)> GetTraitLatentPropertyIndex,
+			TFunctionRef<TConstStructView<FAnimNextTraitSharedData>(uint32 TraitIndex)> GetTraitData);
 
 		// Returns the error state
 		[[nodiscard]] ANIMNEXTANIMGRAPH_API EErrorState GetErrorState() const;
@@ -102,4 +113,4 @@ namespace UE::AnimNext
 		EErrorState ErrorState;
 	};
 }
-#endif
+
diff --git a/Engine/Plugins/Experimental/AnimNextAnimGraph/Source/AnimNextAnimGraph/Internal/Traits/BlendStackTrait.h b/Engine/Plugins/Experimental/AnimNextAnimGraph/Source/AnimNextAnimGraph/Internal/Traits/BlendStackTrait.h
--- a/Engine/Plugins/Experimental/AnimNextAnimGraph/Source/AnimNextAnimGraph/Internal/Traits/BlendStackTrait.h
+++ b/Engine/Plugins/Experimental/AnimNextAnimGraph/Source/AnimNextAnimGraph/Internal/Traits/BlendStackTrait.h
@@ -26,6 +26,12 @@ struct FAnimNextBlendStackCoreTraitSharedData : public FAnimNextTraitSharedData
 	// This struct is empty but required so that we can give a nice display name to the trait
 };
 
+namespace UE::AnimNext
+{
+	// Namespaced alias
+	using FBlendStackCoreTraitData = FAnimNextBlendStackCoreTraitSharedData; 
+}
+
 USTRUCT(meta = (DisplayName = "Blend Stack"))
 struct FAnimNextBlendStackTraitSharedData : public FAnimNextBlendStackCoreTraitSharedData
 {
@@ -190,6 +196,7 @@ namespace UE::AnimNext
 		virtual int32 PushGraph(FExecutionContext& Context, const TTraitBinding<IBlendStack>& Binding, IBlendStack::FGraphRequest&& GraphRequest) const override;
 		virtual int32 GetActiveGraph(FExecutionContext& Context, const TTraitBinding<IBlendStack>& Binding, IBlendStack::FGraphRequestPtr& OutGraphRequest) const override;
 		virtual IBlendStack::FGraphRequestPtr GetGraph(FExecutionContext& Context, const TTraitBinding<IBlendStack>& Binding, int32 InChildIndex) const override;
+		virtual FAnimNextGraphInstance* GetGraphInstance(FExecutionContext& Context, const TTraitBinding<IBlendStack>& Binding, int32 InChildIndex) const override;
 
 		// ISmoothBlendPerBone impl
 		virtual TSharedPtr<const IBlendProfileInterface> GetBlendProfile(FExecutionContext& Context, const TTraitBinding<ISmoothBlendPerBone>& Binding, int32 ChildIndex) const;
diff --git a/Engine/Plugins/Experimental/AnimNextAnimGraph/Source/AnimNextAnimGraph/Private/AnimNextAnimGraphModule.cpp b/Engine/Plugins/Experimental/AnimNextAnimGraph/Source/AnimNextAnimGraph/Private/AnimNextAnimGraphModule.cpp
--- a/Engine/Plugins/Experimental/AnimNextAnimGraph/Source/AnimNextAnimGraph/Private/AnimNextAnimGraphModule.cpp
+++ b/Engine/Plugins/Experimental/AnimNextAnimGraph/Source/AnimNextAnimGraph/Private/AnimNextAnimGraphModule.cpp
@@ -8,7 +8,9 @@
 #include "Modules/ModuleManager.h"
 #include "RigVMCore/RigVMRegistry.h"
 #include "Animation/AnimSequence.h"
+#include "Animation/BlendSpace.h"
 #include "Components/SkeletalMeshComponent.h"
+#include "Factory/AnimGraphFactory.h"
 #include "TraitInterfaces/IUpdate.h"
 #include "Graph/AnimNextAnimationGraph.h"
 #include "Graph/AnimNext_LODPose.h"
@@ -18,6 +20,13 @@
 #include "TraitCore/TraitInterfaceRegistry.h"
 #include "TraitCore/TraitRegistry.h"
 #include "Traits/ModifyCurveTrait.h"
+#include "Traits/BlendSpacePlayerTraitData.h"
+#include "Traits/NotifyDispatcherTraitData.h"
+#include "Traits/SequencePlayer.h"
+#include "Factory/AnimNextFactoryParams.h"
+#include "Injection/InjectionSiteTrait.h"
+#include "Traits/BlendSmoother.h"
+#include "Traits/BlendStackTrait.h"
 
 #if WITH_ANIMNEXT_CONSOLE_COMMANDS
 #include "HAL/IConsoleManager.h"
@@ -31,22 +40,6 @@
 namespace UE::AnimNext::AnimGraph
 {
 
-// Delay loading of default asset factory mappings in settings 
-static FDelayedAutoRegisterHelper AutoRegisterFactoryMappings(EDelayedRegisterRunPhase::EndOfEngineInit, []
-{
-	UAnimNextAnimGraphSettings* Settings = GetMutableDefault<UAnimNextAnimGraphSettings>();
-
-#if WITH_EDITOR
-	// Async load in the editor to avoid stalling startup
-	const bool bLoadAsync = true;
-#else
-	// Outside the editor, we assume we are behind a loading screen of sorts
-	const bool bLoadAsync = false;
-#endif
-
-	Settings->LoadAndRebuildMappings(bLoadAsync);
-});
-
 void FAnimNextAnimGraphModule::StartupModule()
 {
 	// Ensure that AnimNext modules are loaded so we can correctly load plugin content
@@ -84,6 +77,61 @@ void FAnimNextAnimGraphModule::StartupModule()
 	FTraitInterfaceRegistry::Init();
 	FNodeTemplateRegistry::Init();
 
+	// Register asset class -> graph factory mappings
+	{
+		// Register default graph host
+		FAnimNextFactoryParams Params;
+		Params.PushPublicTraitStruct<FBlendStackCoreTraitData>();
+		Params.PushPublicTraitStruct<FBlendSmootherCoreData>();
+		Params.PushPublicTraitStruct<FInjectionSiteTraitData>();
+
+		FAnimGraphFactory::RegisterDefaultParamsForClass(
+			UAnimNextAnimationGraph::StaticClass(),
+			MoveTemp(Params),
+			[](const FFactoryParamsInitializerContext& InContext)
+			{
+				ensure(InContext.AccessStruct<FInjectionSiteTraitData>([&InContext](FInjectionSiteTraitData& InStruct)
+				{
+					InStruct.Graph.Asset = CastChecked<UAnimNextAnimationGraph>(InContext.GetObject());
+				}));
+			});
+	}
+
+	{
+		// Register UAnimSequence 'player'
+		FAnimNextFactoryParams Params;
+		Params.PushPublicTraitStruct<FSequencePlayerData>();
+		Params.PushPublicTraitStruct<FNotifyDispatcherData>();
+		
+		FAnimGraphFactory::RegisterDefaultParamsForClass(
+			UAnimSequence::StaticClass(),
+			MoveTemp(Params),
+			[](const FFactoryParamsInitializerContext& InContext)
+			{
+				ensure(InContext.AccessStruct<FSequencePlayerData>([&InContext](FSequencePlayerData& InStruct)
+				{
+					InStruct.AnimSequence = CastChecked<UAnimSequence>(InContext.GetObject());
+				}));
+			});
+	}
+
+	{
+		// Register UBlendSpace 'player'
+		FAnimNextFactoryParams Params;
+		Params.PushPublicTraitStruct<FBlendSpacePlayerData>();
+		
+		FAnimGraphFactory::RegisterDefaultParamsForClass(
+			UBlendSpace::StaticClass(),
+			MoveTemp(Params),
+			[](const FFactoryParamsInitializerContext& InContext)
+			{
+				ensure(InContext.AccessStruct<FBlendSpacePlayerData>([&InContext](FBlendSpacePlayerData& InStruct)
+				{
+					InStruct.BlendSpace = CastChecked<UBlendSpace>(InContext.GetObject());
+				}));
+			});
+	}
+
 #if WITH_ANIMNEXT_CONSOLE_COMMANDS
 	if (!IsRunningCommandlet())
 	{
diff --git a/Engine/Plugins/Experimental/AnimNextAnimGraph/Source/AnimNextAnimGraph/Private/AnimNextAnimGraphSettings.cpp b/Engine/Plugins/Experimental/AnimNextAnimGraph/Source/AnimNextAnimGraph/Private/AnimNextAnimGraphSettings.cpp
--- a/Engine/Plugins/Experimental/AnimNextAnimGraph/Source/AnimNextAnimGraph/Private/AnimNextAnimGraphSettings.cpp
+++ b/Engine/Plugins/Experimental/AnimNextAnimGraph/Source/AnimNextAnimGraph/Private/AnimNextAnimGraphSettings.cpp
@@ -1,477 +1,11 @@
 // Copyright Epic Games, Inc. All Rights Reserved.
 
 #include "AnimNextAnimGraphSettings.h"
-
-#include "AnimNextAnimGraphModule.h"
-#include "AnimNextDataInterfacePayload.h"
-#include "DataInterface/AnimNextNativeDataInterface.h"
-#include "Graph/AnimNextAnimationGraph.h"
-#include "Logging/StructuredLog.h"
-#include "Misc/EnumerateRange.h"
-#include "Modules/ModuleManager.h"
-
-#if WITH_EDITOR
-#include "AssetRegistry/AssetRegistryHelpers.h"
-#endif // WITH_EDITOR
-
-#if WITH_EDITOR
-void UAnimNextAnimGraphSettings::PostEditChangeProperty(FPropertyChangedEvent& PropertyChangedEvent)
-{
-	if(PropertyChangedEvent.Property)
-	{
-		if( PropertyChangedEvent.Property->GetFName() == GET_MEMBER_NAME_CHECKED(UAnimNextAnimGraphSettings, AssetGraphMappings) ||
-			PropertyChangedEvent.Property->GetFName() == GET_MEMBER_NAME_CHECKED(FAnimNextAssetGraphMapping, AssetType) ||
-			PropertyChangedEvent.Property->GetFName() == GET_MEMBER_NAME_CHECKED(FAnimNextAssetGraphMapping, AnimationGraph) ||
-			PropertyChangedEvent.Property->GetFName() == GET_MEMBER_NAME_CHECKED(FAnimNextAssetGraphMapping, Variable) ||
-			PropertyChangedEvent.Property->GetFName() == GET_MEMBER_NAME_CHECKED(FAnimNextAssetGraphMapping, RequiredDataInterfaces))
-		{
-			LoadAndRebuildMappings(false);
-		}
-	}
-}
-#endif
-
-const UAnimNextAnimationGraph* UAnimNextAnimGraphSettings::GetGraphFromObject(const UObject* InObject) const
-{
-	FAnimNextDataInterfacePayload GraphPayload;
-	return GetGraphFromObject(InObject, GraphPayload);
-}
-
-const UAnimNextAnimationGraph* UAnimNextAnimGraphSettings::GetGraphFromObject(const UObject* InObject, const FAnimNextDataInterfacePayload& InGraphPayload) const
-{
-	check(InObject);
-
-	const UAnimNextAnimationGraph* AnimationGraph = Cast<UAnimNextAnimationGraph>(InObject);
-	if(AnimationGraph)
-	{
-		// Object is an animation graph, so we early out
-		return AnimationGraph;
-	}
-
-	if (!ensure(bMappingsBuiltAtLeastOnce.load(std::memory_order_acquire)))
-	{
-		return nullptr;
-	}
-
-	// Check this class and its parent classes in turn
-	const UClass* Class = InObject->GetClass();
-	int32 Index = INDEX_NONE;
-	while(Class)
-	{
-		const int32* IndexPtr = AssetGraphMap.Find(Class);
-		if(IndexPtr != nullptr)
-		{
-			Index = *IndexPtr;
-			break;
-		}
-
-		Class = Class->GetSuperClass();
-	}
-
-	if(!AssetGraphMappings.IsValidIndex(Index))
-	{
-		return nullptr;
-	}
-
-	TConstArrayView<FInstancedStruct> NativePayloads = InGraphPayload.GetNativePayloads();
-
-	int32 BestNumMatchedInterfaces = -1;
-	int32 BestMappingIndex = INDEX_NONE;
-
-	while (Index != INDEX_NONE)
-	{
-		const FAnimNextAssetGraphMapping& Mapping = AssetGraphMappings[Index];
-
-		bool bIsMappingValid = true;
-		int32 NumMatchedInterfaces = 0;
-
-		for (const TSoftObjectPtr<const UScriptStruct>& DataInterfacePtr : Mapping.RequiredDataInterfaces)
-		{
-			const UScriptStruct* DataInterface = DataInterfacePtr.Get();
-			if (DataInterface == nullptr)
-			{
-				continue;
-			}
-
-			bool bFoundRequiredInterface = false;
-			for (const FInstancedStruct& NativePayload : NativePayloads)
-			{
-				if (NativePayload.GetScriptStruct() == DataInterface)
-				{
-					bFoundRequiredInterface = true;
-					break;
-				}
-			}
-
-			if (bFoundRequiredInterface)
-			{
-				NumMatchedInterfaces++;
-			}
-			else
-			{
-				bIsMappingValid = false;
-				break;
-			}
-		}
-
-		// This mapping is valid, use it if its our best match
-		// If the number of interfaces matched is identical, then we'll pick the first mapping seen,
-		// which will be the last one found through the config files
-		if (bIsMappingValid && NumMatchedInterfaces > BestNumMatchedInterfaces)
-		{
-			BestNumMatchedInterfaces = NumMatchedInterfaces;
-			BestMappingIndex = Index;
-		}
-
-		// Try the next mapping for the same asset type
-		Index = Mapping.NextMappingIndex;
-	}
-
-	// Returns our best match, if one is found
-	return BestMappingIndex != INDEX_NONE ? AssetGraphMappings[BestMappingIndex].AnimationGraph.Get() : nullptr;
-}
-
-bool UAnimNextAnimGraphSettings::CanGetGraphFromAssetClass(const UClass* InClass) const
-{
-	if(InClass == nullptr)
-	{
-		return false;
-	}
-	
-	if(InClass == UAnimNextAnimationGraph::StaticClass())
-	{
-		return true;
-	}
-
-	if (!ensure(bMappingsBuiltAtLeastOnce.load(std::memory_order_acquire)))
-	{
-		return false;
-	}
-
-	const UClass* Class = InClass;
-	while(Class)
-	{
-		if(AssetGraphMap.Find(Class) != nullptr)
-		{
-			return true;
-		}
-
-		Class = Class->GetSuperClass();
-	}
-
-	return false;
-}
-
-FName UAnimNextAnimGraphSettings::GetInjectedVariableNameFromObject(const UObject* InObject) const
-{
-	check(InObject);
-
-	const UAnimNextAnimationGraph* AnimationGraph = Cast<UAnimNextAnimationGraph>(InObject);
-	if(AnimationGraph)
-	{
-		// Object is an animation graph, so we early out
-		return NAME_None;
-	}
-
-	if (!ensure(bMappingsBuiltAtLeastOnce.load(std::memory_order_acquire)))
-	{
-		return NAME_None;
-	}
-
-	// Check this class and its parent classes in turn
-	const UClass* Class = InObject->GetClass();
-	int32 Index = INDEX_NONE;
-	while(Class)
-	{
-		const int32* IndexPtr = AssetGraphMap.Find(Class);
-		if(IndexPtr != nullptr)
-		{
-			Index = *IndexPtr;
-			break;
-		}
-
-		Class = Class->GetSuperClass();
-	}
-
-	if(!AssetGraphMappings.IsValidIndex(Index))
-	{
-		return NAME_None;
-	}
-
-	return AssetGraphMappings[Index].Variable;
-}
-
-
-void UAnimNextAnimGraphSettings::GetNativePayloadFromGraph(const UObject* InObject, const UAnimNextAnimationGraph* InAnimationGraph, FAnimNextDataInterfacePayload& InOutGraphPayload) const
-{
-	GetNativePayloadFromGraph(InObject, InAnimationGraph, InOutGraphPayload.OwnedNativePayloads);
-	InOutGraphPayload.bCombinedPayloadsDirty = true;
-}
-
-void UAnimNextAnimGraphSettings::GetNativePayloadFromGraph(const UObject* InObject, const UAnimNextAnimationGraph* InAnimationGraph, TArray<FInstancedStruct>& InOutGraphPayload) const
-{
-	if(InObject == nullptr || InAnimationGraph == nullptr)
-	{
-		return;
-	}
-
-	auto FindNativeInterface = [&InOutGraphPayload](const UScriptStruct* InNativeInterface)
-	{
-		return InOutGraphPayload.FindByPredicate([InNativeInterface](const FInstancedStruct& InInstancedStruct)
-		{
-			return InInstancedStruct.GetScriptStruct() == InNativeInterface;
-		});
-	};
-
-	// First generate any structs for the supplied graph from the graph's native interfaces, if required
-	for(const FAnimNextImplementedDataInterface& ImplementedInterface : InAnimationGraph->GetImplementedInterfaces())
-	{
-		if(ImplementedInterface.NativeInterface)
-		{
-			// Check if we have an interface already supplied
-			const FInstancedStruct* NativeInterfaceInstance = FindNativeInterface(ImplementedInterface.NativeInterface);
-			if(NativeInterfaceInstance == nullptr)
-			{
-				check(ImplementedInterface.NativeInterface->IsChildOf(FAnimNextNativeDataInterface::StaticStruct()));
-
-				// Create a new native interface
-				FInstancedStruct NewNativeInterface;
-				NewNativeInterface.InitializeAs(ImplementedInterface.NativeInterface);
-
-				FAnimNextNativeDataInterface::FBindToFactoryObjectContext Context;
-				Context.FactoryObject = InObject;
-				Context.DataInterface = InAnimationGraph;
-				NewNativeInterface.GetMutable<FAnimNextNativeDataInterface>().BindToFactoryObject(Context);
-
-				InOutGraphPayload.Add(MoveTemp(NewNativeInterface));
-			}
-		}
-	}
-}
-
-void UAnimNextAnimGraphSettings::GetNonNativePayloadFromGraph(const UObject* InObject, const UAnimNextAnimationGraph* InAnimationGraph, FAnimNextDataInterfacePayload& InOutGraphPayload) const
-{
-	GetNonNativePayloadFromGraph(InObject, InAnimationGraph, InOutGraphPayload.OwnedPayload);
-	InOutGraphPayload.bCombinedPayloadsDirty = true;
-}
-
-void UAnimNextAnimGraphSettings::GetNonNativePayloadFromGraph(const UObject* InObject, const UAnimNextAnimationGraph* InAnimationGraph, FInstancedPropertyBag& InOutGraphPayload) const
-{
-	// Duplicate the public variables
-	InOutGraphPayload = InAnimationGraph->GetPublicVariableDefaults();
-
-	// Setup the factory-mapped variable
-	FName VariableName = GetInjectedVariableNameFromObject(InObject);
-	if(VariableName.IsNone())
-	{
-		return;
-	}
-
-	const FPropertyBagPropertyDesc* Desc = InOutGraphPayload.FindPropertyDescByName(VariableName);
-	if(Desc == nullptr)
-	{
-		return;
-	}
-
-	const FObjectProperty* ObjectProperty = CastField<FObjectProperty>(Desc->CachedProperty);
-	if(ObjectProperty == nullptr)
-	{
-		return;
-	}
-
-#if WITH_EDITOR
-	// Remove editor visibility for the mapped property as UI doesnt want to display it twice
-	const_cast<FObjectProperty*>(ObjectProperty)->PropertyFlags &= ~CPF_Edit;
-#endif
-
-	uint8* ValuePtr = ObjectProperty->ContainerPtrToValuePtr<uint8>(InOutGraphPayload.GetMutableValue().GetMemory());
-	ObjectProperty->SetObjectPropertyValue(ValuePtr, const_cast<UObject*>(InObject));
-}
-
-void UAnimNextAnimGraphSettings::OnDefaultRunGraphHostLoaded(const UAnimNextAnimationGraph* AnimationGraph, bool bLoadAsync)
-{
-	if (AnimationGraph == nullptr)
-	{
-		UE_LOGFMT(LogAnimation, Error, "UAnimNextAnimGraphSettings::LoadAndRebuildMappings: Could not load default host animation graph '{GraphName}'", DefaultRunGraphHost.ToString());
-		return;
-	}
-
-	using namespace UE::AnimNext::AnimGraph;
-	FAnimNextAnimGraphModule& AnimGraphModule = FModuleManager::GetModuleChecked<FAnimNextAnimGraphModule>("AnimNextAnimGraph");
-	AnimGraphModule.LoadedGraphs.Add(AnimationGraph);
-
-	if (AssetGraphMappings.IsEmpty())
-	{
-		FinalizeAsyncLoad();
-		return;
-	}
-
-	for (TEnumerateRef<FAnimNextAssetGraphMapping> Mapping : EnumerateRange(AssetGraphMappings))
-	{
-		FSoftObjectPath AnimationGraphSoftPath = Mapping->AnimationGraph.ToSoftObjectPath();
-
-#if WITH_EDITOR
-		UAssetRegistryHelpers::FixupRedirectedAssetPath(IN OUT AnimationGraphSoftPath);
-#endif
-
-		int32 PackageID = AnimationGraphSoftPath.LoadAsync(FLoadSoftObjectPathAsyncDelegate::CreateLambda(
-			[WeakThis = TWeakObjectPtr<UAnimNextAnimGraphSettings>(this),
-			&Mapping = *Mapping,
-			MappingIndex = Mapping.GetIndex()]
-			(const FSoftObjectPath& InSoftObjectPath, UObject* InObject)
-			{
-				if (UAnimNextAnimGraphSettings* ThisPin = WeakThis.Get())
-				{
-					ThisPin->NumMappingsLoaded++;
-
-					ThisPin->OnMappingAnimationGraphLoaded(CastChecked<const UAnimNextAnimationGraph>(InObject, ECastCheckedType::NullAllowed), Mapping, MappingIndex);
-
-					if (ThisPin->NumMappingsLoaded == ThisPin->AssetGraphMappings.Num())
-					{
-						ThisPin->FinalizeAsyncLoad();
-					}
-				}
-			}));
-
-		if (!bLoadAsync)
-		{
-			FlushAsyncLoading(PackageID);
-		}
-	}
-}
-
-void UAnimNextAnimGraphSettings::OnMappingAnimationGraphLoaded(const UAnimNextAnimationGraph* AnimationGraph, FAnimNextAssetGraphMapping& Mapping, int32 MappingIndex)
-{
-	if (AnimationGraph == nullptr)
-	{
-		UE_LOGFMT(LogAnimation, Error, "UAnimNextAnimGraphSettings::LoadAndRebuildMappings: Could not load animation graph '{GraphName}'", Mapping.AnimationGraph.ToString());
-		return;
-	}
-
-	using namespace UE::AnimNext::AnimGraph;
-	FAnimNextAnimGraphModule& AnimGraphModule = FModuleManager::GetModuleChecked<FAnimNextAnimGraphModule>("AnimNextAnimGraph");
-	AnimGraphModule.LoadedGraphs.Add(AnimationGraph);
-
-	const UClass* AssetType = Mapping.AssetType.LoadSynchronous();
-	if (AssetType == nullptr)
-	{
-		UE_LOGFMT(LogAnimation, Error, "UAnimNextAnimGraphSettings::LoadAndRebuildMappings: Could not load asset class '{ClassName}'", Mapping.AssetType.ToString());
-		return;
-	}
-
-	bool bAreRequiredInterfacesValid = true;
-	for (const TSoftObjectPtr<const UScriptStruct>& DataInterfacePtr : Mapping.RequiredDataInterfaces)
-	{
-		const UScriptStruct* DataInterface = DataInterfacePtr.LoadSynchronous();
-		if (DataInterface == nullptr)
-		{
-			UE_LOGFMT(LogAnimation, Error, "UAnimNextAnimGraphSettings::LoadAndRebuildMappings: Could not load required data interface class '{ClassName}'", DataInterfacePtr.ToString());
-			bAreRequiredInterfacesValid = false;
-			continue;
-		}
-	}
-
-	if (!bAreRequiredInterfacesValid)
-	{
-		return;
-	}
-
-	// Find or add our new entry
-	int32& FirstAssetMappingIndex = AssetGraphMap.FindOrAdd(AssetType, INDEX_NONE);
-
-	// Update the next mapping index to point to our old entry (or INDEX_NONE if we are the first)
-	Mapping.NextMappingIndex = FirstAssetMappingIndex;
-
-	// Our new first mapping index is the current one
-	FirstAssetMappingIndex = MappingIndex;
-
-	// Warn if the variable is not settable in the public interface of the graph
-	if (AnimationGraph != nullptr && AssetType != nullptr && !Mapping.Variable.IsNone())
-	{
-		const FInstancedPropertyBag& PublicVariableDefaults = AnimationGraph->GetPublicVariableDefaults();
-		const FPropertyBagPropertyDesc* Desc = PublicVariableDefaults.FindPropertyDescByName(Mapping.Variable);
-		if (Desc == nullptr)
-		{
-			UE_LOGFMT(LogAnimation, Error, "UAnimNextAnimGraphSettings::LoadAndRebuildMappings: Could not find public variable '{VariableName}' in graph '{GraphName}'", Mapping.Variable, AnimationGraph->GetPathName());
-			return;
-		}
-
-		const FObjectProperty* ObjectProperty = CastField<FObjectProperty>(Desc->CachedProperty);
-		if (ObjectProperty == nullptr)
-		{
-			UE_LOGFMT(LogAnimation, Error, "UAnimNextAnimGraphSettings::LoadAndRebuildMappings: Variable '{VariableName}' in graph '{GraphName}' is not of object type", Mapping.Variable, AnimationGraph->GetPathName());
-			return;
-		}
-
-		if (!AssetType->IsChildOf(ObjectProperty->PropertyClass))
-		{
-			UE_LOGFMT(LogAnimation, Error, "UAnimNextAnimGraphSettings::LoadAndRebuildMappings: Variable '{VariableName}' in graph '{GraphName}' is not of a compatible object type: '{TargetType}' vs '{SourceType}'", Mapping.Variable, AnimationGraph->GetPathName(), ObjectProperty->PropertyClass->GetFullName(), AssetType->GetFullName());
-			return;
-		}
-	}
-}
-
-void UAnimNextAnimGraphSettings::FinalizeAsyncLoad()
-{
-	// Release memory order to ensure all previous writes complete before other threads can see this value change
-	// Queries to this value must use the acquire memory order to ensure loads are not re-ordered before the value
-	// is read and tested
-	bMappingsBuiltAtLeastOnce.store(true, std::memory_order_release);
-}
-
-void UAnimNextAnimGraphSettings::LoadAndRebuildMappings(bool bLoadAsync)
-{
-	using namespace UE::AnimNext::AnimGraph;
-
-	// To avoid loading every asset synchronously at engine init, we load them async in multiple steps
-	//   - We first load the default run graph host
-	//   - Followed by the mapping graph assets concurrently
-	//   - Once everything has loaded, the graph factories are usable
-	
-	FAnimNextAnimGraphModule& AnimGraphModule = FModuleManager::GetModuleChecked<FAnimNextAnimGraphModule>("AnimNextAnimGraph");
-	AnimGraphModule.LoadedGraphs.Reset();
-	AssetGraphMap.Reset();
-
-	int32 PackageID = DefaultRunGraphHost.ToSoftObjectPath().LoadAsync(FLoadSoftObjectPathAsyncDelegate::CreateLambda([WeakThis = TWeakObjectPtr<UAnimNextAnimGraphSettings>(this), bLoadAsync](const FSoftObjectPath& InSoftObjectPath, UObject* InObject)
-		{
-			if (TStrongObjectPtr<UAnimNextAnimGraphSettings> ThisPin = WeakThis.Pin())
-			{
-				ThisPin->OnDefaultRunGraphHostLoaded(CastChecked<const UAnimNextAnimationGraph>(InObject), bLoadAsync);
-			}
-		}));
-
-	if (!bLoadAsync)
-	{
-		FlushAsyncLoading(PackageID);
-	}
-}
+#include "Factory/AnimGraphFactory.h"
 
 TArray<UClass*> UAnimNextAnimGraphSettings::GetAllowedAssetClasses()
 {
-	return GetDefault<UAnimNextAnimGraphSettings>()->GetAllowedAssetClassesImpl();
-}
-
-TArray<UClass*> UAnimNextAnimGraphSettings::GetAllowedAssetClassesImpl() const
-{
-	TArray<UClass*> AllowedClasses;
-
-	if (!ensure(bMappingsBuiltAtLeastOnce.load(std::memory_order_acquire)))
-	{
-		return AllowedClasses;
-	}
-
-	// Add all derived classes of UAnimNextAnimationGraph
-	AllowedClasses.Add(UAnimNextAnimationGraph::StaticClass());
-	GetDerivedClasses(UAnimNextAnimationGraph::StaticClass(), AllowedClasses);
-
-	// Add all mappings
-	for(const FAnimNextAssetGraphMapping& Mapping : AssetGraphMappings)
-	{
-		if(UClass* AssetClass = Mapping.AssetType.Get())
-		{
-			AllowedClasses.Add(AssetClass);
-		}
-	}
-
-	return AllowedClasses;
+	TArray<UClass*> Array;
+	Array = UE::AnimNext::FAnimGraphFactory::GetRegisteredClasses();
+	return Array;
 }
diff --git a/Engine/Plugins/Experimental/AnimNextAnimGraph/Source/AnimNextAnimGraph/Private/EvaluationVM/Tasks/PushAnimSequenceKeyframe.cpp b/Engine/Plugins/Experimental/AnimNextAnimGraph/Source/AnimNextAnimGraph/Private/EvaluationVM/Tasks/PushAnimSequenceKeyframe.cpp
--- a/Engine/Plugins/Experimental/AnimNextAnimGraph/Source/AnimNextAnimGraph/Private/EvaluationVM/Tasks/PushAnimSequenceKeyframe.cpp
+++ b/Engine/Plugins/Experimental/AnimNextAnimGraph/Source/AnimNextAnimGraph/Private/EvaluationVM/Tasks/PushAnimSequenceKeyframe.cpp
@@ -9,7 +9,7 @@
 #include "EvaluationVM/EvaluationVM.h"
 #include "EvaluationVM/KeyframeState.h"
 
-FAnimNextAnimSequenceKeyframeTask FAnimNextAnimSequenceKeyframeTask::MakeFromSampleTime(TWeakObjectPtr<UAnimSequence> AnimSequence, double SampleTime, bool bInterpolate)
+FAnimNextAnimSequenceKeyframeTask FAnimNextAnimSequenceKeyframeTask::MakeFromSampleTime(TWeakObjectPtr<const UAnimSequence> AnimSequence, double SampleTime, bool bInterpolate)
 {
 	FAnimNextAnimSequenceKeyframeTask Task;
 	Task.AnimSequence = AnimSequence;
@@ -19,7 +19,7 @@ FAnimNextAnimSequenceKeyframeTask FAnimNextAnimSequenceKeyframeTask::MakeFromSam
 	return Task;
 }
 
-FAnimNextAnimSequenceKeyframeTask FAnimNextAnimSequenceKeyframeTask::MakeFromKeyframeIndex(TWeakObjectPtr<UAnimSequence> AnimSequence, uint32 KeyframeIndex)
+FAnimNextAnimSequenceKeyframeTask FAnimNextAnimSequenceKeyframeTask::MakeFromKeyframeIndex(TWeakObjectPtr<const UAnimSequence> AnimSequence, uint32 KeyframeIndex)
 {
 	FAnimNextAnimSequenceKeyframeTask Task;
 	Task.AnimSequence = AnimSequence;
diff --git a/Engine/Plugins/Experimental/AnimNextAnimGraph/Source/AnimNextAnimGraph/Private/Factory/AnimGraphFactory.cpp b/Engine/Plugins/Experimental/AnimNextAnimGraph/Source/AnimNextAnimGraph/Private/Factory/AnimGraphFactory.cpp
new file mode 100644
--- /dev/null
+++ b/Engine/Plugins/Experimental/AnimNextAnimGraph/Source/AnimNextAnimGraph/Private/Factory/AnimGraphFactory.cpp
@@ -0,0 +1,132 @@
+﻿// Copyright Epic Games, Inc. All Rights Reserved.
+
+#include "Factory/AnimGraphFactory.h"
+
+#include "Factory/AnimGraphBuilderContext.h"
+#include "Containers/StripedMap.h"
+#include "Factory/AnimNextAnimGraphBuilder.h"
+#include "Graph/AnimNextAnimationGraph.h"
+#include "TraitCore/TraitWriter.h"
+
+namespace UE::AnimNext::Private
+{
+
+// References to all generated factory graphs, keyed by method
+static TStripedMap<32, uint64, TWeakObjectPtr<const UAnimNextAnimationGraph>> GAnimationGraphMap;
+
+// All registered object -> param mappings
+static TMap<const UClass*, TPair<FAnimNextFactoryParams, FAnimGraphFactory::FParamsInitializer>> GDefaultParamsForObject;
+static TArray<UClass*> GAllFactoryClasses;
+
+}
+
+namespace UE::AnimNext
+{
+
+uint64 IAnimGraphBuilder::GetKey() const
+{
+	return 0;
+}
+
+const UAnimNextAnimationGraph* FAnimGraphFactory::GetDefaultGraphHost()
+{
+	const TPair<FAnimNextFactoryParams, FParamsInitializer>* ParamsTask = Private::GDefaultParamsForObject.Find(UAnimNextAnimationGraph::StaticClass());
+	return BuildGraph(ParamsTask->Key.GetBuilder());
+}
+
+const UAnimNextAnimationGraph* FAnimGraphFactory::GetOrBuildGraph(const UObject* InObject, const IAnimGraphBuilder& InBuilder)
+{
+	if (const UAnimNextAnimationGraph* AnimationGraph = Cast<UAnimNextAnimationGraph>(InObject))
+	{
+		return AnimationGraph;
+	}
+
+	return BuildGraph(InBuilder);
+}
+
+const UAnimNextAnimationGraph* FAnimGraphFactory::BuildGraph(const IAnimGraphBuilder& InBuilder)
+{
+	uint64 Key = InBuilder.GetKey();
+	if (Key == 0)
+	{
+		return nullptr;
+	}
+
+	auto ProduceGraph = [&InBuilder]() -> TWeakObjectPtr<const UAnimNextAnimationGraph>
+	{
+		FAnimGraphBuilderContext Context;
+		if (InBuilder.Build(Context))
+		{
+			return Context.Build();
+		}
+		return nullptr;
+	};
+
+	TWeakObjectPtr<const UAnimNextAnimationGraph> AnimationGraph;
+	auto ApplyWeakGraph = [&ProduceGraph, &AnimationGraph](TWeakObjectPtr<const UAnimNextAnimationGraph>& InFoundAnimationGraph)
+	{
+		if (InFoundAnimationGraph.IsStale())
+		{
+			// Found a stale weak ptr, so regenerate (this may be a procedural graph or a previously loaded graph that has been GCed)
+			InFoundAnimationGraph = ProduceGraph();
+		}
+		AnimationGraph = InFoundAnimationGraph;
+	};
+
+	Private::GAnimationGraphMap.FindOrProduceAndApplyForWrite(Key, ProduceGraph, ApplyWeakGraph);
+	return AnimationGraph.Get();
+}
+
+const FAnimNextFactoryParams& FAnimGraphFactory::GetDefaultParamsForClass(const UClass* InClass)
+{
+	check(InClass);
+	const TPair<FAnimNextFactoryParams, FParamsInitializer>* ParamsTask = Private::GDefaultParamsForObject.Find(InClass);
+	check(ParamsTask != nullptr);	// Unregistered object type
+	return ParamsTask->Key;
+}
+
+FAnimNextFactoryParams FAnimGraphFactory::GetDefaultParamsForObject(const UObject* InObject)
+{
+	check(InObject);
+	const TPair<FAnimNextFactoryParams, FParamsInitializer>* ParamsTask = Private::GDefaultParamsForObject.Find(InObject->GetClass());
+	check(ParamsTask != nullptr);	// Unregistered object type
+
+	FAnimNextFactoryParams ParamsCopy = ParamsTask->Key;
+	if (ParamsCopy.Builder.Stacks.Num() > 0)
+	{
+		FFactoryParamsInitializerContext Context;
+		Context.Object = InObject;
+		Context.StructDatas = ParamsCopy.Builder.Stacks[0].TraitStructs;
+		ParamsTask->Value(Context);
+	}
+	return ParamsCopy;
+}
+
+void FAnimGraphFactory::InitializeDefaultParamsForObject(const UObject* InObject, FAnimNextFactoryParams& InOutParams)
+{
+	check(InObject);
+	const TPair<FAnimNextFactoryParams, FParamsInitializer>* ParamsTask = Private::GDefaultParamsForObject.Find(InObject->GetClass());
+	check(ParamsTask != nullptr);	// Unregistered object type
+
+	if (InOutParams.Builder.Stacks.Num() > 0)
+	{
+		FFactoryParamsInitializerContext Context;
+		Context.Object = InObject;
+		Context.StructDatas = InOutParams.Builder.Stacks[0].TraitStructs;
+		ParamsTask->Value(Context);
+	}
+}
+
+void FAnimGraphFactory::RegisterDefaultParamsForClass(const UClass* InClass, FAnimNextFactoryParams&& InParams, FParamsInitializer&& InInitializer)
+{
+	Private::GDefaultParamsForObject.Add(InClass, { MoveTemp(InParams), MoveTemp(InInitializer) });
+	Private::GAllFactoryClasses.Add(const_cast<UClass*>(InClass));
+}
+
+TConstArrayView<UClass*> FAnimGraphFactory::GetRegisteredClasses()
+{
+	return Private::GAllFactoryClasses;
+}
+
+
+}
diff --git a/Engine/Plugins/Experimental/AnimNextAnimGraph/Source/AnimNextAnimGraph/Private/Factory/AnimNextFactoryParams.cpp b/Engine/Plugins/Experimental/AnimNextAnimGraph/Source/AnimNextAnimGraph/Private/Factory/AnimNextFactoryParams.cpp
new file mode 100644
--- /dev/null
+++ b/Engine/Plugins/Experimental/AnimNextAnimGraph/Source/AnimNextAnimGraph/Private/Factory/AnimNextFactoryParams.cpp
@@ -0,0 +1,31 @@
+﻿// Copyright Epic Games, Inc. All Rights Reserved.
+
+#include "Factory/AnimNextFactoryParams.h"
+#include "Graph/AnimNextGraphInstance.h"
+
+void FAnimNextFactoryParams::InitializeInstance(FUAFAssetInstance& InInstance) const
+{
+	if (Builder.Stacks.Num() > 0)
+	{
+		for (const TInstancedStruct<FAnimNextTraitSharedData>& InitializePayload : Builder.Stacks[0].TraitStructs)
+		{
+			InInstance.AccessVariablesStruct(InitializePayload.GetScriptStruct(), [&InitializePayload](FStructView InStructView)
+			{
+				check(InitializePayload.GetScriptStruct() == InStructView.GetScriptStruct());
+				InStructView.GetScriptStruct()->CopyScriptStruct(InStructView.GetMemory(), InitializePayload.GetMemory());
+			});
+		}
+	}
+
+	const UE::AnimNext::FInstanceTaskContext TaskContext(InInstance);
+	for (const UE::AnimNext::FInstanceTask& InitializeTask : InitializeTasks)
+	{
+		InitializeTask(TaskContext);
+	}
+}
+
+void FAnimNextFactoryParams::Reset()
+{
+	Builder.Reset();
+	InitializeTasks.Empty();
+}
\ No newline at end of file
diff --git a/Engine/Plugins/Experimental/AnimNextAnimGraph/Source/AnimNextAnimGraph/Private/Factory/AnimNextSimpleAnimGraphBuilder.cpp b/Engine/Plugins/Experimental/AnimNextAnimGraph/Source/AnimNextAnimGraph/Private/Factory/AnimNextSimpleAnimGraphBuilder.cpp
new file mode 100644
--- /dev/null
+++ b/Engine/Plugins/Experimental/AnimNextAnimGraph/Source/AnimNextAnimGraph/Private/Factory/AnimNextSimpleAnimGraphBuilder.cpp
@@ -0,0 +1,221 @@
+﻿// Copyright Epic Games, Inc. All Rights Reserved.
+
+#include "Factory/AnimNextSimpleAnimGraphBuilder.h"
+
+#include "Factory/AnimGraphBuilderContext.h"
+#include "Misc/HashBuilder.h"
+#include "TraitCore/NodeHandle.h"
+#include "TraitCore/NodeTemplate.h"
+#include "TraitCore/TraitRegistry.h"
+#include "TraitCore/NodeTemplateBuilder.h"
+#include "TraitCore/TraitWriter.h"
+#include "Traits/ReferencePoseTrait.h"
+
+bool FAnimNextSimpleAnimGraphBuilder::Build(UE::AnimNext::FAnimGraphBuilderContext& InContext) const
+{
+	using namespace UE::AnimNext;
+
+	ensure(Stacks.Num() == 1);	// Currently only support a single stack
+
+	if (Stacks.Num() != 1 || Stacks[0].TraitStructs.Num() == 0)
+	{
+		return false;
+	}
+
+	const FTraitRegistry& TraitRegistry = FTraitRegistry::Get();
+
+	InContext.VariableStructs.Reserve(Stacks[0].TraitStructs.Num());
+
+	struct FTraitStackData
+	{
+		FNodeHandle NodeHandle;
+		TArray<TInstancedStruct<FAnimNextTraitSharedData>, TInlineAllocator<4>> TraitStructs;
+		TArray<FTraitUID, TInlineAllocator<4>> TraitUIDs;
+		TMap<TPair<uint32, FName>, int32, TInlineSetAllocator<64>> VariableIndexMap;
+	};
+
+	TArray<FTraitStackData, TInlineAllocator<4>> TraitStacks;
+	TraitStacks.Reserve(Stacks.Num());
+
+	{
+		FTraitStackData TraitStackData;
+		TraitStackData.TraitStructs = Stacks[0].TraitStructs;
+		TraitStackData.TraitUIDs.Reserve(Stacks[0].TraitStructs.Num());
+		TraitStacks.Add(MoveTemp(TraitStackData));
+	}
+
+	for (int32 TraitStackIndex = 0; TraitStackIndex < TraitStacks.Num(); ++TraitStackIndex)
+	{
+		FTraitStackData& TraitStackData = TraitStacks[TraitStackIndex];
+		int32 LatentPropertyIndex = 0;
+		int32 TraitIndex = 0;
+		for (TInstancedStruct<FAnimNextTraitSharedData>& TraitStruct : TraitStackData.TraitStructs)
+		{
+			const FTrait* Trait = TraitRegistry.Find(TraitStruct.GetScriptStruct());
+			TraitStackData.TraitUIDs.Add(Trait->GetTraitUID());
+			if (TraitStackIndex == 0)
+			{
+				// Add the variables of the single trait stack we support as 'public interface' structs
+				InContext.VariableStructs.Add(TraitStruct.GetScriptStruct());
+			}
+
+			for (TFieldIterator<FProperty> It(TraitStruct.GetScriptStruct()); It; ++It)
+			{
+				TraitStackData.VariableIndexMap.Add({TraitIndex, It->GetFName()}, LatentPropertyIndex);
+				LatentPropertyIndex++;
+			}
+			TraitIndex++;
+		}
+	}
+
+	// Register nodes
+	auto RegisterNode = [&InContext](FTraitStackData& InTraitStackData)
+	{
+		TArray<uint8> NodeTemplateBuffer;
+		const FNodeTemplate* NodeTemplate = FNodeTemplateBuilder::BuildNodeTemplate(InTraitStackData.TraitUIDs, NodeTemplateBuffer);
+		InTraitStackData.NodeHandle = InContext.TraitWriter.RegisterNode(*NodeTemplate);
+	};
+
+	const int32 NumTraitStacks = TraitStacks.Num();
+	for (int32 TraitStackIndex = 0; TraitStackIndex < NumTraitStacks; ++TraitStackIndex)
+	{
+		FTraitStackData& TraitStackData = TraitStacks[TraitStackIndex];
+		RegisterNode(TraitStackData);
+
+		// Single trait we support is always the root
+		if (TraitStackIndex == 0)
+		{
+			InContext.RootTraitHandle = FAnimNextEntryPointHandle(TraitStackData.NodeHandle);
+		}
+
+		// Add new trait stack with refpose trait to cover 'unlinked' child pins
+		const FTrait* Trait = TraitRegistry.Find(FAnimNextReferencePoseTraitSharedData::StaticStruct());
+		const FTraitUID RefPoseUID = Trait->GetTraitUID();
+		for (TInstancedStruct<FAnimNextTraitSharedData>& TraitStruct : TraitStackData.TraitStructs)
+		{
+			for (TFieldIterator<FProperty> It(TraitStruct.GetScriptStruct()); It; ++It)
+			{
+				FStructProperty* StructProperty = CastField<FStructProperty>(*It);
+				if (StructProperty && StructProperty->Struct == FAnimNextTraitHandle::StaticStruct())
+				{
+					FTraitStackData RefPoseTraitStackData;
+					TInstancedStruct<FAnimNextTraitSharedData> InstancedStruct;
+					InstancedStruct.InitializeAsScriptStruct(FAnimNextReferencePoseTraitSharedData::StaticStruct());
+
+					RefPoseTraitStackData.TraitStructs.Add(MoveTemp(InstancedStruct));
+					RefPoseTraitStackData.TraitUIDs.Add(RefPoseUID);
+
+					RegisterNode(RefPoseTraitStackData);
+
+					// Link to this new stack
+					FAnimNextTraitHandle* Handle = It->ContainerPtrToValuePtr<FAnimNextTraitHandle>(TraitStruct.GetMutableMemory());
+					*Handle = FAnimNextTraitHandle(RefPoseTraitStackData.NodeHandle);
+
+					TraitStacks.Add(RefPoseTraitStackData);
+				}
+			}
+		}
+	}
+
+	// Write traits
+	InContext.TraitWriter.BeginNodeWriting();
+	for (int32 TraitStackIndex = 0; TraitStackIndex < TraitStacks.Num(); ++TraitStackIndex)
+	{
+		const FTraitStackData& TraitStackData = TraitStacks[TraitStackIndex];
+		InContext.TraitWriter.WriteNode(TraitStackData.NodeHandle,
+			[&TraitStackData](uint32 InTraitIndex, FName InPropertyName)
+			{
+				return TraitStackData.VariableIndexMap.FindRef({InTraitIndex, InPropertyName});
+			},
+			[&TraitStackData](uint32 InTraitIndex)
+			{
+				return TConstStructView<FAnimNextTraitSharedData>(TraitStackData.TraitStructs[InTraitIndex]);
+			});
+	}
+	InContext.TraitWriter.EndNodeWriting();
+
+	ensure(InContext.TraitWriter.GetErrorState() == FTraitWriter::EErrorState::None);
+
+	return true;
+}
+
+void FAnimNextSimpleAnimGraphBuilder::ValidateTraitStruct(int32 InStackIndex, TConstStructView<FAnimNextTraitSharedData> InStruct)
+{
+	using namespace UE::AnimNext;
+
+	ensure(InStackIndex == 0);	// Currently only support a single stack
+	if (InStackIndex != 0)
+	{
+		return;
+	}
+
+	if (Stacks.Num() == 0)
+	{
+		Stacks.SetNum(InStackIndex + 1);
+	}
+
+	const FTraitRegistry& TraitRegistry = FTraitRegistry::Get();
+	const FTrait* Trait = TraitRegistry.Find(InStruct.GetScriptStruct());
+	if (Trait == nullptr)
+	{
+		ensureMsgf(false, TEXT("FAnimNextSimpleAnimGraphBuilder: Could not find registered trait"));
+		return;
+	}
+
+	if (Stacks[InStackIndex].TraitStructs.Num() == 0 && Trait->GetTraitMode() != ETraitMode::Base)
+	{
+		ensureMsgf(false, TEXT("FAnimNextSimpleAnimGraphBuilder: First trait should always be a base trait"));
+		return;
+	}
+}
+
+void FAnimNextSimpleAnimGraphBuilder::PushTraitStructViewToStack(int32 InStackIndex, TConstStructView<FAnimNextTraitSharedData> InStruct)
+{
+	using namespace UE::AnimNext;
+
+	ValidateTraitStruct(InStackIndex, InStruct);
+
+	auto FindExistingTrait = [&InStruct](const TInstancedStruct<FAnimNextTraitSharedData>& InExistingStruct)
+	{
+		return InStruct.GetScriptStruct() == InExistingStruct.GetScriptStruct();
+	};
+
+	if (TInstancedStruct<FAnimNextTraitSharedData>* ExistingStruct = Stacks[InStackIndex].TraitStructs.FindByPredicate(FindExistingTrait))
+	{
+		*ExistingStruct = InStruct;
+	}
+	else
+	{
+		Stacks[InStackIndex].TraitStructs.Emplace(InStruct);
+	}
+
+	InvalidateKey();
+}
+
+void FAnimNextSimpleAnimGraphBuilder::PushTraitInstancedStructToStack(int32 InStackIndex, TInstancedStruct<FAnimNextTraitSharedData>&& InStruct)
+{
+	using namespace UE::AnimNext;
+
+	ValidateTraitStruct(InStackIndex, InStruct);
+
+	auto FindExistingTrait = [&InStruct](const TInstancedStruct<FAnimNextTraitSharedData>& InExistingStruct)
+	{
+		return InStruct.GetScriptStruct() == InExistingStruct.GetScriptStruct();
+	};
+
+	if (TInstancedStruct<FAnimNextTraitSharedData>* ExistingStruct = Stacks[InStackIndex].TraitStructs.FindByPredicate(FindExistingTrait))
+	{
+		*ExistingStruct = InStruct;
+	}
+	else
+	{
+		Stacks[InStackIndex].TraitStructs.Emplace(MoveTemp(InStruct));
+	}
+
+	InvalidateKey();
+}
+
+uint64 FAnimNextSimpleAnimGraphBuilder::RecalculateKey() const
+{
+	return UE::StructUtils::GetStructHash64(TConstStructView<FAnimNextSimpleAnimGraphBuilder>(*this));
+}
diff --git a/Engine/Plugins/Experimental/AnimNextAnimGraph/Source/AnimNextAnimGraph/Private/Graph/AnimNextAnimationGraph.cpp b/Engine/Plugins/Experimental/AnimNextAnimGraph/Source/AnimNextAnimGraph/Private/Graph/AnimNextAnimationGraph.cpp
--- a/Engine/Plugins/Experimental/AnimNextAnimGraph/Source/AnimNextAnimGraph/Private/Graph/AnimNextAnimationGraph.cpp
+++ b/Engine/Plugins/Experimental/AnimNextAnimGraph/Source/AnimNextAnimGraph/Private/Graph/AnimNextAnimationGraph.cpp
@@ -35,11 +35,13 @@ UAnimNextAnimationGraph::UAnimNextAnimationGraph(const FObjectInitializer& Objec
 	ExtendedExecuteContext.SetContextPublicDataStruct(FAnimNextExecuteContext::StaticStruct());
 }
 
-TSharedPtr<FAnimNextGraphInstance> UAnimNextAnimationGraph::AllocateInstance(FAnimNextModuleInstance* InModuleInstance, UE::AnimNext::FExecutionContext* InParentContext, FAnimNextGraphInstance* InParentGraphInstance, FName InEntryPoint) const
+TSharedPtr<FAnimNextGraphInstance> UAnimNextAnimationGraph::AllocateInstance(const UE::AnimNext::FGraphAllocationParams& InParams) const
 {
 	SCOPE_CYCLE_COUNTER(STAT_AnimNext_Graph_AllocateInstance);
 
-	const FName EntryPoint = (InEntryPoint == NAME_None) ? DefaultEntryPoint : InEntryPoint;
+	PRAGMA_DISABLE_DEPRECATION_WARNINGS
+	const FName EntryPoint = (InParams.EntryPoint == NAME_None) ? DefaultEntryPoint : InParams.EntryPoint;
+	PRAGMA_ENABLE_DEPRECATION_WARNINGS
 	const FAnimNextTraitHandle ResolvedRootTraitHandle = ResolvedRootTraitHandles.FindRef(EntryPoint);
 	if (!ResolvedRootTraitHandle.IsValid())
 	{
@@ -47,51 +49,27 @@ TSharedPtr<FAnimNextGraphInstance> UAnimNextAnimationGraph::AllocateInstance(FAn
 	}
 
 	TSharedPtr<FAnimNextGraphInstance> InstanceImpl = MakeShared<FAnimNextGraphInstance>();
-	InstanceImpl->DataInterface = this;
-	InstanceImpl->ModuleInstance = InModuleInstance;
-	InstanceImpl->HostInstance = InParentGraphInstance != nullptr ? static_cast<FAnimNextDataInterfaceInstance*>(InParentGraphInstance) : static_cast<FAnimNextDataInterfaceInstance*>(InModuleInstance);
+	InstanceImpl->Asset = this;
+	InstanceImpl->ModuleInstance = InParams.ModuleInstance;
+	InstanceImpl->HostInstance = InParams.ParentGraphInstance != nullptr ? static_cast<FUAFAssetInstance*>(InParams.ParentGraphInstance) : static_cast<FUAFAssetInstance*>(InParams.ModuleInstance);
 	InstanceImpl->EntryPoint = EntryPoint;
 
 	// If we have a parent graph, use its root since we share the same root, otherwise if we have no parent, we are the root
-	InstanceImpl->RootGraphInstance = InParentGraphInstance != nullptr ? InParentGraphInstance->GetRootGraphInstance() : InstanceImpl.Get();
+	InstanceImpl->RootGraphInstance = InParams.ParentGraphInstance != nullptr ? InParams.ParentGraphInstance->GetRootGraphInstance() : InstanceImpl.Get();
 
-	InstanceImpl->Variables = VariableDefaults;
-	InstanceImpl->ExtendedExecuteContext = ExtendedExecuteContext;
-	if(VariableDefaults.GetPropertyBagStruct() != nullptr)
-	{
-		InstanceImpl->PublicVariablesState = FAnimNextGraphInstance::EPublicVariablesState::Unbound;
-
-		// Setup external variables memory ptrs manually as we dont follow the pattern of owning multiple URigVMHosts like control rig.
-		// InitializeVM() is called, but only sets up handles for the defaults in the module, not for an instance
-		const int32 NumVariables = InstanceImpl->Variables.GetNumPropertiesInBag();
-		TArray<FRigVMExternalVariableRuntimeData> ExternalVariableRuntimeData;
-		ExternalVariableRuntimeData.Reserve(NumVariables);
-		TConstArrayView<FPropertyBagPropertyDesc> Descs = InstanceImpl->Variables.GetPropertyBagStruct()->GetPropertyDescs();
-		uint8* BasePtr = InstanceImpl->Variables.GetMutableValue().GetMemory();
-		for(int32 VariableIndex = 0; VariableIndex < NumVariables; ++VariableIndex)
-		{
-			ExternalVariableRuntimeData.Emplace(Descs[VariableIndex].CachedProperty->ContainerPtrToValuePtr<uint8>(BasePtr));
-		}
-		InstanceImpl->ExtendedExecuteContext.ExternalVariableRuntimeData = MoveTemp(ExternalVariableRuntimeData);
-	}
-	else
-	{
-		InstanceImpl->PublicVariablesState = FAnimNextGraphInstance::EPublicVariablesState::None;
-	}
-
-	// Now initialize the 'instance', cache memory handles etc. in the context
-	VM->InitializeInstance(InstanceImpl->ExtendedExecuteContext);
+	// Set up variable's data
+	InstanceImpl->InitializeVariables(InParams.Overrides);
 
 	{
 		UE::AnimNext::FExecutionContext Context(*InstanceImpl.Get());
 
-		if (InParentContext != nullptr)
+		if (InParams.ParentContext != nullptr)
 		{
-			Context.SetBindingObject(InParentContext->GetBindingObject());
+			Context.SetBindingObject(InParams.ParentContext->GetBindingObject());
 		}
-		else if (InModuleInstance != nullptr)
+		else if (InParams.ModuleInstance != nullptr)
 		{
-			const FAnimNextSkeletalMeshComponentReferenceComponent& ComponentReference = InModuleInstance->GetComponent<FAnimNextSkeletalMeshComponentReferenceComponent>();
+			const FAnimNextSkeletalMeshComponentReferenceComponent& ComponentReference = InParams.ModuleInstance->GetComponent<FAnimNextSkeletalMeshComponentReferenceComponent>();
 
 			UE::AnimNext::FDataHandle RefPoseHandle = UE::AnimNext::FDataRegistry::Get()->GetOrGenerateReferencePose(ComponentReference.GetComponent());
 			const UE::AnimNext::FReferencePose& RefPose = RefPoseHandle.GetRef<UE::AnimNext::FReferencePose>();
