UAF: Fix crashes when editing dependent assets

Reworks compilation logic to be batched. Repurposes the scoped compiler results to be a 'compile job' scope. Only once all dependent assets (recompiled via their modification delegates) are recompiled to any instances get re-allocated (via the OnCompileFinished delegate). This prevents instances that are nested inside hosts from having out-of-data instance data compared to their host.
Reworks combined property bag generation to use source data rather than compiled data, which prevents circular dependency issues from occuring.

diff --git a/Engine/Plugins/Experimental/AnimNext/Source/AnimNext/Internal/Component/AnimNextWorldSubsystem.h b/Engine/Plugins/Experimental/AnimNext/Source/AnimNext/Internal/Component/AnimNextWorldSubsystem.h
--- a/Engine/Plugins/Experimental/AnimNext/Source/AnimNext/Internal/Component/AnimNextWorldSubsystem.h
+++ b/Engine/Plugins/Experimental/AnimNext/Source/AnimNext/Internal/Component/AnimNextWorldSubsystem.h
@@ -103,8 +103,8 @@ private:
 
 protected:
 #if WITH_EDITOR
-	// Refresh any entries that use the provided module as it has been recompiled.
-	UE_API void OnModuleCompiled(UAnimNextModule* InModule);
+	// Refresh any entries that use the provided asset as it has been recompiled.
+	UE_API void OnCompileJobFinished(UAnimNextRigVMAsset* InAsset);
 #endif
 
 	// Register a handle to the subsystem
@@ -208,8 +208,8 @@ protected:
 	FRWLock PendingLock;
 
 #if WITH_EDITOR
-	// Handle used to hook module compilation
-	FDelegateHandle OnModuleCompiledHandle;
+	// Handle used to hook asset compilation
+	FDelegateHandle OnCompileJobFinishedHandle;
 #endif
 
 	// Handle used to hook into pre-world tick
diff --git a/Engine/Plugins/Experimental/AnimNext/Source/AnimNext/Internal/Module/AnimNextModule.h b/Engine/Plugins/Experimental/AnimNext/Source/AnimNext/Internal/Module/AnimNextModule.h
--- a/Engine/Plugins/Experimental/AnimNext/Source/AnimNext/Internal/Module/AnimNextModule.h
+++ b/Engine/Plugins/Experimental/AnimNext/Source/AnimNext/Internal/Module/AnimNextModule.h
@@ -57,12 +57,6 @@ public:
 	UE_API virtual void Serialize(FArchive& Ar) override;
 	UE_API virtual void PostLoad() override;
 
-#if WITH_EDITORONLY_DATA
-	// Delegate called in editor when a module is compiled
-	using FOnModuleCompiled = TTSMulticastDelegate<void(UAnimNextModule*)>;
-	static UE_API FOnModuleCompiled& OnModuleCompiled();
-#endif
-
 protected:
 	friend class UAnimNextModuleFactory;
 	friend class UAnimNextModule_EditorData;
diff --git a/Engine/Plugins/Experimental/AnimNext/Source/AnimNext/Internal/Module/AnimNextModuleInstance.h b/Engine/Plugins/Experimental/AnimNext/Source/AnimNext/Internal/Module/AnimNextModuleInstance.h
--- a/Engine/Plugins/Experimental/AnimNext/Source/AnimNext/Internal/Module/AnimNextModuleInstance.h
+++ b/Engine/Plugins/Experimental/AnimNext/Source/AnimNext/Internal/Module/AnimNextModuleInstance.h
@@ -143,7 +143,7 @@ private:
 
 #if WITH_EDITOR
 	// Resets internal state if the module we are bound to is recompiled in editor
-	void OnModuleCompiled();
+	void OnCompileJobFinished();
 #endif
 
 	// Give the module components a chance to handle events and call their OnEndExecution extension point
diff --git a/Engine/Plugins/Experimental/AnimNext/Source/AnimNext/Private/AnimNextRigVMAsset.cpp b/Engine/Plugins/Experimental/AnimNext/Source/AnimNext/Private/AnimNextRigVMAsset.cpp
--- a/Engine/Plugins/Experimental/AnimNext/Source/AnimNext/Private/AnimNextRigVMAsset.cpp
+++ b/Engine/Plugins/Experimental/AnimNext/Source/AnimNext/Private/AnimNextRigVMAsset.cpp
@@ -12,6 +12,14 @@
 #include "Modules/ModuleManager.h"
 #include "UObject/AssetRegistryTagsContext.h"
 
+namespace UE::AnimNext::Private
+{
+#if WITH_EDITOR
+	static UAnimNextRigVMAsset::FOnCompileJobEvent GOnCompileJobStarted;
+	static UAnimNextRigVMAsset::FOnCompileJobEvent GOnCompileJobFinished;
+#endif
+}
+
 UAnimNextRigVMAsset::UAnimNextRigVMAsset(const FObjectInitializer& ObjectInitializer)
 	: Super(ObjectInitializer)
 {
@@ -178,4 +186,18 @@ UE::AnimNext::FFunctionHandle UAnimNextRigVMAsset::GetFunctionHandle(FName InEve
 	}
 
 	return FFunctionHandle();
-}
\ No newline at end of file
+}
+
+#if WITH_EDITOR
+
+UAnimNextRigVMAsset::FOnCompileJobEvent& UAnimNextRigVMAsset::OnCompileJobStarted()
+{
+	return UE::AnimNext::Private::GOnCompileJobStarted;
+}
+
+UAnimNextRigVMAsset::FOnCompileJobEvent& UAnimNextRigVMAsset::OnCompileJobFinished()
+{
+	return UE::AnimNext::Private::GOnCompileJobStarted;
+}
+
+#endif
\ No newline at end of file
diff --git a/Engine/Plugins/Experimental/AnimNext/Source/AnimNext/Private/Component/AnimNextWorldSubsystem.cpp b/Engine/Plugins/Experimental/AnimNext/Source/AnimNext/Private/Component/AnimNextWorldSubsystem.cpp
--- a/Engine/Plugins/Experimental/AnimNext/Source/AnimNext/Private/Component/AnimNextWorldSubsystem.cpp
+++ b/Engine/Plugins/Experimental/AnimNext/Source/AnimNext/Private/Component/AnimNextWorldSubsystem.cpp
@@ -15,7 +15,7 @@ UAnimNextWorldSubsystem::UAnimNextWorldSubsystem()
 	if(!HasAnyFlags(RF_ClassDefaultObject))
 	{
 #if WITH_EDITOR
-		OnModuleCompiledHandle = UAnimNextModule::OnModuleCompiled().AddUObject(this, &UAnimNextWorldSubsystem::OnModuleCompiled);
+		OnCompileJobFinishedHandle = UAnimNextRigVMAsset::OnCompileJobFinished().AddUObject(this, &UAnimNextWorldSubsystem::OnCompileJobFinished);
 #endif
 
 		// Kick off root task at the start of each world tick
@@ -38,7 +38,7 @@ void UAnimNextWorldSubsystem::BeginDestroy()
 	if (!HasAnyFlags(RF_ClassDefaultObject))
 	{
 #if WITH_EDITOR
-		UAnimNextModule::OnModuleCompiled().Remove(OnModuleCompiledHandle);
+		UAnimNextRigVMAsset::OnCompileJobFinished().Remove(OnCompileJobFinishedHandle);
 #endif
 
 		FWorldDelegates::OnWorldPreActorTick.Remove(OnWorldPreActorTickHandle);
@@ -239,16 +239,16 @@ void UAnimNextWorldSubsystem::QueueTaskHandle(UE::AnimNext::FModuleHandle InOutH
 
 #if WITH_EDITOR
 
-void UAnimNextWorldSubsystem::OnModuleCompiled(UAnimNextModule* InModule)
+void UAnimNextWorldSubsystem::OnCompileJobFinished(UAnimNextRigVMAsset* InAsset)
 {
 	// Cant do this while we are running in a world tick
 	check(!GetWorld()->bInTick); 
 
 	for(FAnimNextModuleInstance& Instance : Instances)
 	{
-		if(Instance.GetModule() == InModule)
+		if(Instance.GetModule() == InAsset)
 		{
-			Instance.OnModuleCompiled();
+			Instance.OnCompileJobFinished();
 		}
 	}
 }
diff --git a/Engine/Plugins/Experimental/AnimNext/Source/AnimNext/Private/Module/AnimNextModule.cpp b/Engine/Plugins/Experimental/AnimNext/Source/AnimNext/Private/Module/AnimNextModule.cpp
--- a/Engine/Plugins/Experimental/AnimNext/Source/AnimNext/Private/Module/AnimNextModule.cpp
+++ b/Engine/Plugins/Experimental/AnimNext/Source/AnimNext/Private/Module/AnimNextModule.cpp
@@ -16,10 +16,6 @@ REGISTER_ASSETDEPENDENCY_GATHERER(FExternalAssetDependencyGatherer, UAnimNextMod
 
 #include UE_INLINE_GENERATED_CPP_BY_NAME(AnimNextModule)
 
-#if WITH_EDITORONLY_DATA
-static UAnimNextModule::FOnModuleCompiled GOnModuleCompiled; 
-#endif
-
 UAnimNextModule::UAnimNextModule(const FObjectInitializer& ObjectInitializer)
 	: Super(ObjectInitializer)
 {
@@ -58,12 +54,3 @@ void UAnimNextModule::PostLoad()
 	}
 #endif
 }
-
-#if WITH_EDITOR
-
-UAnimNextModule::FOnModuleCompiled& UAnimNextModule::OnModuleCompiled()
-{
-	return GOnModuleCompiled;
-}
-
-#endif
\ No newline at end of file
diff --git a/Engine/Plugins/Experimental/AnimNext/Source/AnimNext/Private/Module/AnimNextModuleInstance.cpp b/Engine/Plugins/Experimental/AnimNext/Source/AnimNext/Private/Module/AnimNextModuleInstance.cpp
--- a/Engine/Plugins/Experimental/AnimNext/Source/AnimNext/Private/Module/AnimNextModuleInstance.cpp
+++ b/Engine/Plugins/Experimental/AnimNext/Source/AnimNext/Private/Module/AnimNextModuleInstance.cpp
@@ -281,7 +281,7 @@ void FAnimNextModuleInstance::Initialize()
 
 		// Initialize variables
 #if WITH_EDITOR
-		if(bIsRecreatingOnCompile)
+		if(bIsRecreatingOnCompile && Variables.bHasBeenInitialized)
 		{
 			MigrateVariables();
 		}
@@ -486,7 +486,7 @@ const UAnimNextModule* FAnimNextModuleInstance::GetModule() const
 }
 
 #if WITH_EDITOR
-void FAnimNextModuleInstance::OnModuleCompiled()
+void FAnimNextModuleInstance::OnCompileJobFinished()
 {
 	using namespace UE::AnimNext;
 
diff --git a/Engine/Plugins/Experimental/AnimNext/Source/AnimNext/Private/Variables/UAFInstanceVariableContainer.cpp b/Engine/Plugins/Experimental/AnimNext/Source/AnimNext/Private/Variables/UAFInstanceVariableContainer.cpp
--- a/Engine/Plugins/Experimental/AnimNext/Source/AnimNext/Private/Variables/UAFInstanceVariableContainer.cpp
+++ b/Engine/Plugins/Experimental/AnimNext/Source/AnimNext/Private/Variables/UAFInstanceVariableContainer.cpp
@@ -1033,15 +1033,23 @@ void FUAFInstanceVariableContainer::ResolveOverrides()
 #if WITH_EDITOR
 void FUAFInstanceVariableContainer::Migrate()
 {
+	using namespace UE::AnimNext;
+	
 	switch (VariablesContainer.GetIndex())
 	{
 	case FVariablesContainerType::IndexOfType<FInstancedPropertyBag>():
 		{
-			VariablesContainer.Get<FInstancedPropertyBag>().MigrateToNewBagInstance(AssetOrStructData.Get<FAssetType>()->GetVariableDefaults());
+			FInstancedPropertyBag& PropertyBag = VariablesContainer.Get<FInstancedPropertyBag>();
+			PropertyBag.MigrateToNewBagInstance(AssetOrStructData.Get<FAssetType>()->GetVariableDefaults());
+			NumVariables = PropertyBag.GetNumPropertiesInBag();
 			break;
 		}
 	case FVariablesContainerType::IndexOfType<FInstancedStruct>():
-		break;
+		{
+			FInstancedStruct& InstancedStruct = VariablesContainer.Get<FInstancedStruct>();
+			NumVariables = FStructDataCache::GetStructInfo(InstancedStruct.GetScriptStruct())->GetProperties().Num();
+			break;
+		}
 	default:
 		checkNoEntry();
 		break;
diff --git a/Engine/Plugins/Experimental/AnimNext/Source/AnimNext/Private/Variables/UAFInstanceVariableData.cpp b/Engine/Plugins/Experimental/AnimNext/Source/AnimNext/Private/Variables/UAFInstanceVariableData.cpp
--- a/Engine/Plugins/Experimental/AnimNext/Source/AnimNext/Private/Variables/UAFInstanceVariableData.cpp
+++ b/Engine/Plugins/Experimental/AnimNext/Source/AnimNext/Private/Variables/UAFInstanceVariableData.cpp
@@ -166,7 +166,23 @@ void FUAFInstanceVariableData::Initialize(const FUAFAssetInstance& InInstance, c
 			}
 		}
 	}
+
+	RebuildNameMaps();
 	
+	ResolveOverrides();
+
+#if WITH_EDITOR
+	bHasBeenInitialized = true;
+#endif
+}
+
+void FUAFInstanceVariableData::RebuildNameMaps()
+{
+	using namespace UE::AnimNext;
+
+	LegacyVariableNameMap.Reset();
+	VariablesMap.Reset();
+
 	// Setup name maps
 	const int32 NumSharedVariableSets = AllSharedVariableContainers.Num();
 	for (int32 SharedVariableSetIndex = 0; SharedVariableSetIndex < NumSharedVariableSets; ++SharedVariableSetIndex)
@@ -215,8 +231,6 @@ void FUAFInstanceVariableData::Initialize(const FUAFAssetInstance& InInstance, c
 			break;
 		}
 	}
-
-	ResolveOverrides();
 }
 
 void FUAFInstanceVariableData::Reset()
@@ -234,18 +248,61 @@ void FUAFInstanceVariableData::Reset()
 #if WITH_EDITOR
 void FUAFInstanceVariableData::Migrate()
 {
-	for (TSharedRef<FUAFInstanceVariableContainer>& SharedVariableSet : OwnedVariableContainers)
+	// Higher level logic needs to check to see if we have already been initialized, as overrides need to be cached first
+	check(bHasBeenInitialized);
+
+	NumAllVariables = 0;
+	NumInternalVariables = 0;
+	AllSharedVariableContainersMap.Reset();
+
+	for (TSharedRef<FUAFInstanceVariableContainer>& OwnedVariableContainer : OwnedVariableContainers)
+	{
+		OwnedVariableContainer->Migrate();
+	}
+
+	int32 ContainerIndex = 0;
+	for (TWeakPtr<FUAFInstanceVariableContainer>& WeakSharedVariableContainer : AllSharedVariableContainers)
+	{
+		TSharedPtr<FUAFInstanceVariableContainer> SharedVariableContainer = WeakSharedVariableContainer.Pin();
+		check(SharedVariableContainer.IsValid());
+
+		NumAllVariables += SharedVariableContainer->NumVariables;
+
+		const UObject* Object = nullptr;
+		switch (SharedVariableContainer->AssetOrStructData.GetIndex())
+		{
+		case FUAFInstanceVariableContainer::FAssetOrStructType::IndexOfType<FUAFInstanceVariableContainer::FAssetType>():
+			Object = SharedVariableContainer->AssetOrStructData.Get<FUAFInstanceVariableContainer::FAssetType>();
+			break;
+		case FUAFInstanceVariableContainer::FAssetOrStructType::IndexOfType<FUAFInstanceVariableContainer::FStructType>():
+			Object = SharedVariableContainer->AssetOrStructData.Get<FUAFInstanceVariableContainer::FStructType>();
+			break;
+		default:
+			checkNoEntry()
+			break;
+		}
+
+		AllSharedVariableContainersMap.Add(Object, ContainerIndex++);
+	}
+
+	RebuildNameMaps();
+
+	for (TWeakPtr<FUAFInstanceVariableContainer>& WeakInternalSharedVariableContainer : InternalSharedVariableContainers)
 	{
-		SharedVariableSet->Migrate();
+		TSharedPtr<FUAFInstanceVariableContainer> InternalSharedVariableContainer = WeakInternalSharedVariableContainer.Pin();
+		check(InternalSharedVariableContainer.IsValid());
+		
+		NumAllVariables += InternalSharedVariableContainer->NumVariables;
+		NumInternalVariables += InternalSharedVariableContainer->NumVariables;
 	}
 }
 #endif
 
 void FUAFInstanceVariableData::ResolveOverrides()
 {
-	for (TSharedRef<FUAFInstanceVariableContainer>& SharedVariableSet : OwnedVariableContainers)
+	for (TSharedRef<FUAFInstanceVariableContainer>& OwnedVariableContainer : OwnedVariableContainers)
 	{
-		SharedVariableSet->ResolveOverrides();
+		OwnedVariableContainer->ResolveOverrides();
 	}
 }
 
diff --git a/Engine/Plugins/Experimental/AnimNext/Source/AnimNext/Public/AnimNextRigVMAsset.h b/Engine/Plugins/Experimental/AnimNext/Source/AnimNext/Public/AnimNextRigVMAsset.h
--- a/Engine/Plugins/Experimental/AnimNext/Source/AnimNext/Public/AnimNextRigVMAsset.h
+++ b/Engine/Plugins/Experimental/AnimNext/Source/AnimNext/Public/AnimNextRigVMAsset.h
@@ -74,6 +74,15 @@ public:
 	// Get a function handle given its mangled event name
 	UE_API UE::AnimNext::FFunctionHandle GetFunctionHandle(FName InEventName) const;
 
+#if WITH_EDITOR
+	// Delegate called in editor when an asset's compilation has started (i.e. asset data that instances rely on is going to change)
+	using FOnCompileJobEvent = TTSMulticastDelegate<void(UAnimNextRigVMAsset*)>;
+	static UE_API FOnCompileJobEvent& OnCompileJobStarted();
+
+	// Delegate called in editor when an asset's compilation has finished (and instances can be re-allocated)
+	static UE_API FOnCompileJobEvent& OnCompileJobFinished();
+#endif
+	
 private:
 	// URigVMHost interface
 	UE_API virtual TArray<FRigVMExternalVariable> GetExternalVariablesImpl(bool bFallbackToBlueprint) const;
diff --git a/Engine/Plugins/Experimental/AnimNext/Source/AnimNext/Public/Variables/UAFInstanceVariableData.h b/Engine/Plugins/Experimental/AnimNext/Source/AnimNext/Public/Variables/UAFInstanceVariableData.h
--- a/Engine/Plugins/Experimental/AnimNext/Source/AnimNext/Public/Variables/UAFInstanceVariableData.h
+++ b/Engine/Plugins/Experimental/AnimNext/Source/AnimNext/Public/Variables/UAFInstanceVariableData.h
@@ -11,6 +11,7 @@ class UAnimNextRigVMAsset;
 class UAnimNextSharedVariables;
 struct FUAFAssetInstance;
 struct FAnimNextModuleInstance;
+struct FAnimNextGraphInstance;
 struct FAnimNextParamType;
 struct FUAFInstanceVariableData;
 struct FUAFInstanceVariableContainerProxy;
@@ -42,6 +43,7 @@ struct FUAFInstanceVariableData
 private:
 	friend FUAFAssetInstance;
 	friend FAnimNextModuleInstance;
+	friend FAnimNextGraphInstance;
 	friend FUAFInstanceVariableDataProxy;
 	friend UE::AnimNext::FAnimNextTrace;
 	friend UE::AnimNext::FInstanceTaskContext;
@@ -53,6 +55,9 @@ private:
 	// Reset to the default state
 	void Reset();
 
+	// Helper function used during initialization & migration
+	void RebuildNameMaps();
+
 #if WITH_EDITOR
 	// Used during compilation, migrate all variables to new bags according to new defaults.
 	void Migrate();
@@ -137,6 +142,11 @@ private:
 
 	// Total count of all variables used to operate the graph/module
 	int32 NumInternalVariables = 0;
+
+#if WITH_EDITOR
+	// Whether Initialize() has been called at least once
+	bool bHasBeenInitialized = false;
+#endif
 };
 
 template<>
diff --git a/Engine/Plugins/Experimental/AnimNext/Source/AnimNextEditor/Private/AssetCompilationHandler.cpp b/Engine/Plugins/Experimental/AnimNext/Source/AnimNextEditor/Private/AssetCompilationHandler.cpp
--- a/Engine/Plugins/Experimental/AnimNext/Source/AnimNextEditor/Private/AssetCompilationHandler.cpp
+++ b/Engine/Plugins/Experimental/AnimNext/Source/AnimNextEditor/Private/AssetCompilationHandler.cpp
@@ -3,7 +3,7 @@
 #include "AssetCompilationHandler.h"
 #include "AnimNextRigVMAsset.h"
 #include "AnimNextRigVMAssetEditorData.h"
-#include "AnimNextScopedCompilerResults.h"
+#include "AnimNextScopedCompileJob.h"
 #include "UncookedOnlyUtils.h"
 
 namespace UE::AnimNext::Editor
@@ -23,11 +23,11 @@ void FAssetCompilationHandler::Compile(TSharedRef<Workspace::IWorkspaceEditor> I
 	UAnimNextRigVMAsset* AnimNextRigVMAsset = CastChecked<UAnimNextRigVMAsset>(InAsset);
 	UAnimNextRigVMAssetEditorData* EditorData = UE::AnimNext::UncookedOnly::FUtils::GetEditorData<UAnimNextRigVMAssetEditorData>(AnimNextRigVMAsset);
 
-	UncookedOnly::FScopedCompilerResults CompilerResults(InAsset);
+	UncookedOnly::FScopedCompileJob CompileJob(InAsset);
 	EditorData->RecompileVM();
 
-	NumErrors = CompilerResults.GetLog().NumErrors;
-	NumWarnings = CompilerResults.GetLog().NumWarnings;
+	NumErrors = CompileJob.GetLog().NumErrors;
+	NumWarnings = CompileJob.GetLog().NumWarnings;
 
 	OnCompileStatusChanged().ExecuteIfBound();
 }
diff --git a/Engine/Plugins/Experimental/AnimNext/Source/AnimNextEditor/Private/Workspace/AnimNextWorkspaceEditorMode.cpp b/Engine/Plugins/Experimental/AnimNext/Source/AnimNextEditor/Private/Workspace/AnimNextWorkspaceEditorMode.cpp
--- a/Engine/Plugins/Experimental/AnimNext/Source/AnimNextEditor/Private/Workspace/AnimNextWorkspaceEditorMode.cpp
+++ b/Engine/Plugins/Experimental/AnimNext/Source/AnimNextEditor/Private/Workspace/AnimNextWorkspaceEditorMode.cpp
@@ -6,7 +6,7 @@
 #include "AnimNextEditorModule.h"
 #include "AnimNextRigVMAsset.h"
 #include "AnimNextRigVMAssetEditorData.h"
-#include "AnimNextScopedCompilerResults.h"
+#include "AnimNextScopedCompileJob.h"
 #include "AnimNextWorkspaceSchema.h"
 #include "ContextObjectStore.h"
 #include "IAssetCompilationHandler.h"
@@ -188,7 +188,7 @@ void UAnimNextWorkspaceEditorMode::HandleCompile()
 	}
 
 	// Start a batch compile scope
-	UE::AnimNext::UncookedOnly::FScopedCompilerResults CompileResults(JobName, Asset, Assets);
+	UE::AnimNext::UncookedOnly::FScopedCompileJob CompileJob(JobName, Assets);
 
 	UE::AnimNext::Editor::FAnimNextEditorModule& AnimNextEditorModule = FModuleManager::LoadModuleChecked<UE::AnimNext::Editor::FAnimNextEditorModule>("AnimNextEditor");
 	for(UObject* AssetToCompile : Assets)
diff --git a/Engine/Plugins/Experimental/AnimNext/Source/AnimNextUncookedOnly/Internal/AnimNextRigVMAssetEditorData.h b/Engine/Plugins/Experimental/AnimNext/Source/AnimNextUncookedOnly/Internal/AnimNextRigVMAssetEditorData.h
--- a/Engine/Plugins/Experimental/AnimNext/Source/AnimNextUncookedOnly/Internal/AnimNextRigVMAssetEditorData.h
+++ b/Engine/Plugins/Experimental/AnimNext/Source/AnimNextUncookedOnly/Internal/AnimNextRigVMAssetEditorData.h
@@ -33,7 +33,7 @@ namespace UE::AnimNext::UncookedOnly
 	struct FPublicVariablesImpl;
 	struct FUtils;
 	struct FUtilsPrivate;
-	class FScopedCompilerResults;
+	class FScopedCompileJob;
 }
 
 namespace UE::AnimNext::Editor
@@ -214,7 +214,7 @@ protected:
 	friend UE::AnimNext::Editor::FAnimNextAnimGraphEditorModule;
 	friend class UE::AnimNext::Editor::FCallFunctionSharedDataDetails;
 	friend UAnimNextWorkspaceEditorMode;
-	friend UE::AnimNext::UncookedOnly::FScopedCompilerResults;
+	friend UE::AnimNext::UncookedOnly::FScopedCompileJob;
 	friend UE::AnimNext::Editor::FAssetCompilationHandler;
 	friend class UAnimNextStateTreeTreeEditorData;
 	friend struct FAnimNextModuleTests;
@@ -335,6 +335,9 @@ protected:
 	////////////////////////////////////////////////////////////////////////////////////////////////////////////////////////
 	// Begin Compilation overrides, in order of operation
 
+	// Called before a batch of compile jobs for dependent asset starts
+	virtual void OnCompileJobStarted() {}
+
 	// Compilation phase 1: Called before RigVM compilation to setup compiler settings and clean our outer asset of compiler-generated data
 	virtual void OnPreCompileAsset(FRigVMCompileSettings& InSettings) {}
 
@@ -345,6 +348,10 @@ protected:
 	// Compilation phase 3: Called before RigVM compilation to allow this asset to generate variables to be injected, separate method to allow programmatic graphs to use these vars
 	// These variables will be regenerated each compile, and are not saved between compiles
 	virtual void OnPreCompileGetProgrammaticVariables(const FRigVMCompileSettings& InSettings, FAnimNextGetVariableCompileContext& OutCompileContext) {}
+	
+	// Compilation phase 3.1: Called after variable compilation, before RigVM compilation to allow this asset to generate variables to be injected, separate method to allow programmatic graphs to use these vars
+	// These variables will be regenerated each compile, and are not saved between compiles
+	virtual void OnPostCompileVariables(const FRigVMCompileSettings& InSettings, const FAnimNextGetVariableCompileContext& InCompileContext) {}
 
 	// Compilation phase 4: Called before RigVM compilation to allow this asset to generate graphs to be injected
 	virtual void OnPreCompileGetProgrammaticGraphs(const FRigVMCompileSettings& InSettings, FAnimNextGetGraphCompileContext& OutCompileContext) {}
@@ -355,6 +362,9 @@ protected:
 	// Compilation phase 6: Called after RigVM compilation to clean up/finish the compilation process
 	virtual void OnPostCompileCleanup(const FRigVMCompileSettings& InSettings) {}
 
+	// Called after a batch of compile jobs for dependent asset finishes and re-allocation can occur
+	virtual void OnCompileJobFinished() {}
+
 	////////////////////////////////////////////////////////////////////////////////////////////////////////////////////////
 	// End Compilation overrides
 
@@ -551,8 +561,7 @@ protected:
 	void BuildFunctionWrapperEvents(FAnimNextRigVMAssetCompileContext& InContext, const FRigVMCompileSettings& InSettings);
 
 	// Generates a combined property bag consisting of all variables both internal and shared
-	// NOTE: Uses compiled data from the hosting asset uses and thus can introduce circular dependency issues if misused!
-	UE_API FInstancedPropertyBag GenerateCombinedPropertyBag() const;
+	UE_API FInstancedPropertyBag GenerateCombinedPropertyBag(const FAnimNextGetVariableCompileContext& InCompileContext) const;
 
 	// Serialization/post-load fixup for old data
 	void UpgradeDataInterfaces();
diff --git a/Engine/Plugins/Experimental/AnimNext/Source/AnimNextUncookedOnly/Internal/Entries/AnimNextRigVMAssetEntry.h b/Engine/Plugins/Experimental/AnimNext/Source/AnimNextUncookedOnly/Internal/Entries/AnimNextRigVMAssetEntry.h
--- a/Engine/Plugins/Experimental/AnimNext/Source/AnimNextUncookedOnly/Internal/Entries/AnimNextRigVMAssetEntry.h
+++ b/Engine/Plugins/Experimental/AnimNext/Source/AnimNextUncookedOnly/Internal/Entries/AnimNextRigVMAssetEntry.h
@@ -56,6 +56,7 @@ public:
 #if WITH_EDITOR
 	UE_API virtual void PostEditChangeProperty(FPropertyChangedEvent& PropertyChangedEvent) override;
 	UE_API virtual void PostTransacted(const FTransactionObjectEvent& TransactionEvent) override;
+	UE_API virtual void PostEditUndo() override;
 #endif
 
 protected:
diff --git a/Engine/Plugins/Experimental/AnimNext/Source/AnimNextUncookedOnly/Internal/Module/AnimNextModule_EditorData.h b/Engine/Plugins/Experimental/AnimNext/Source/AnimNextUncookedOnly/Internal/Module/AnimNextModule_EditorData.h
--- a/Engine/Plugins/Experimental/AnimNext/Source/AnimNextUncookedOnly/Internal/Module/AnimNextModule_EditorData.h
+++ b/Engine/Plugins/Experimental/AnimNext/Source/AnimNextUncookedOnly/Internal/Module/AnimNextModule_EditorData.h
@@ -53,7 +53,6 @@ private:
 	// UAnimNextRigVMAssetEditorData interface
 	virtual UScriptStruct* GetExecuteContextStruct() const override { return FAnimNextExecuteContext::StaticStruct(); }
 	virtual TConstArrayView<TSubclassOf<UAnimNextRigVMAssetEntry>> GetEntryClasses() const override;
-	virtual void RecompileVM() override;
 	virtual void OnPreCompileAsset(FRigVMCompileSettings& InSettings) override;
 	virtual void OnPreCompileGetProgrammaticGraphs(const FRigVMCompileSettings& InSettings, FAnimNextGetGraphCompileContext& OutCompileContext) override;
 	virtual void OnPreCompileProcessGraphs(const FRigVMCompileSettings& InSettings, FAnimNextProcessGraphCompileContext& OutCompileContext) override;
diff --git a/Engine/Plugins/Experimental/AnimNext/Source/AnimNextUncookedOnly/Internal/UncookedOnlyUtils.h b/Engine/Plugins/Experimental/AnimNext/Source/AnimNextUncookedOnly/Internal/UncookedOnlyUtils.h
--- a/Engine/Plugins/Experimental/AnimNext/Source/AnimNextUncookedOnly/Internal/UncookedOnlyUtils.h
+++ b/Engine/Plugins/Experimental/AnimNext/Source/AnimNextUncookedOnly/Internal/UncookedOnlyUtils.h
@@ -186,6 +186,8 @@ struct FUtils
 {
 	static UE_API void CompileVariables(const FRigVMCompileSettings& InSettings, UAnimNextRigVMAsset* InAsset, FAnimNextGetVariableCompileContext& OutCompileContext);
 
+	static UE_API FInstancedPropertyBag MakePropertyBagForEditorData(const UAnimNextRigVMAssetEditorData* InEditorData, const FAnimNextGetVariableCompileContext& InCompileContext);
+	
 	static UE_API void CompileVariableBindings(const FRigVMCompileSettings& InSettings, UAnimNextRigVMAsset* InAsset, TArray<URigVMGraph*>& OutGraphs);
 
 	static UE_API void RecreateVM(UAnimNextRigVMAsset* InAsset);
diff --git a/Engine/Plugins/Experimental/AnimNext/Source/AnimNextUncookedOnly/Private/AnimNextRigVMAssetEditorData.cpp b/Engine/Plugins/Experimental/AnimNext/Source/AnimNextUncookedOnly/Private/AnimNextRigVMAssetEditorData.cpp
--- a/Engine/Plugins/Experimental/AnimNext/Source/AnimNextUncookedOnly/Private/AnimNextRigVMAssetEditorData.cpp
+++ b/Engine/Plugins/Experimental/AnimNext/Source/AnimNextUncookedOnly/Private/AnimNextRigVMAssetEditorData.cpp
@@ -13,7 +13,7 @@
 #include "AnimNextRigVMAssetSchema.h"
 #include "AnimNextAssetWorkspaceAssetUserData.h"
 #include "AnimNextControllerBase.h"
-#include "AnimNextScopedCompilerResults.h"
+#include "AnimNextScopedCompileJob.h"
 #include "RigVMPythonUtils.h"
 #include "ExternalPackageHelper.h"
 #include "IAnimNextRigVMGraphInterface.h"
@@ -54,12 +54,18 @@
 
 void UAnimNextRigVMAssetEditorData::BroadcastModified(EAnimNextEditorDataNotifType InType, UObject* InSubject)
 {
-	RequestAutoVMRecompilation();
+	using namespace UE::AnimNext::UncookedOnly;
+
+	// Modifications here can trigger compilation, so add a scope to catch compile batches
+	UAnimNextRigVMAsset* Asset = FUtils::GetAsset<UAnimNextRigVMAsset>(this);
+	FScopedCompileJob CompilerResults(LOCTEXT("ModifiedAssetsJobName", "Modified Assets"), { Asset });
 
 	if(!bSuspendEditorDataNotifications)
 	{
 		ModifiedDelegate.Broadcast(this, InType, InSubject);
 	}
+
+	RequestAutoVMRecompilation();
 }
 
 void UAnimNextRigVMAssetEditorData::ReportError(const TCHAR* InMessage)
@@ -988,7 +994,7 @@ void UAnimNextRigVMAssetEditorData::RecompileVM()
 	TGuardValue<bool> CompilingGuard(bIsCompiling, true);
 
 	UAnimNextRigVMAsset* Asset = FUtils::GetAsset<UAnimNextRigVMAsset>(this);
-	FScopedCompilerResults CompilerResults(Asset);
+	FScopedCompileJob CompilerResults(Asset);
 
 	VMCompileSettings.SetExecuteContextStruct(FAnimNextExecuteContext::StaticStruct());
 	FRigVMCompileSettings Settings = (bCompileInDebugMode) ? FRigVMCompileSettings::Fast(VMCompileSettings.GetExecuteContextStruct()) : VMCompileSettings;
@@ -1025,6 +1031,7 @@ void UAnimNextRigVMAssetEditorData::RecompileVM()
 		{
 			FAnimNextGetVariableCompileContext GetVariableCompileContext(CompileContext);
 			FUtils::CompileVariables(Settings, Asset, GetVariableCompileContext);
+			OnPostCompileVariables(Settings, GetVariableCompileContext);
 		}
 
 		BuildFunctionWrapperEvents(CompileContext, Settings);
@@ -2382,7 +2389,7 @@ void UAnimNextRigVMAssetEditorData::HandleReportFromCompiler(EMessageSeverity::T
 		return;
 	}
 
-	FCompilerResultsLog& Log = UE::AnimNext::UncookedOnly::FScopedCompilerResults::GetLog();
+	FCompilerResultsLog& Log = UE::AnimNext::UncookedOnly::FScopedCompileJob::GetLog();
 
 	UObject* SubjectForMessage = InSubject;
 	if(URigVMNode* ModelNode = Cast<URigVMNode>(SubjectForMessage))
@@ -2715,11 +2722,12 @@ void UAnimNextRigVMAssetEditorData::SetUseExternalPackages(TArrayView<UAnimNextR
 
 #endif
 
-FInstancedPropertyBag UAnimNextRigVMAssetEditorData::GenerateCombinedPropertyBag() const
+FInstancedPropertyBag UAnimNextRigVMAssetEditorData::GenerateCombinedPropertyBag(const FAnimNextGetVariableCompileContext& InCompileContext) const
 {
 	UAnimNextRigVMAsset* Asset = UE::AnimNext::UncookedOnly::FUtils::GetAsset<UAnimNextRigVMAsset>(this);
 
 	FInstancedPropertyBag CombinedPropertyBag;
+
 	auto AddInstancedPropertyBag = [this, &CombinedPropertyBag](const FInstancedPropertyBag& InPropertyBag)
 	{
 		if(const UPropertyBag* PropertyBag = InPropertyBag.GetPropertyBagStruct())
@@ -2736,7 +2744,8 @@ FInstancedPropertyBag UAnimNextRigVMAssetEditorData::GenerateCombinedPropertyBag
 	// It must match that of the variables FUAFInstanceVariableData::Initialize
 
 	// First add internal variables
-	AddInstancedPropertyBag(Asset->VariableDefaults);
+	FInstancedPropertyBag AssetBag = UE::AnimNext::UncookedOnly::FUtils::MakePropertyBagForEditorData(this, InCompileContext);
+	AddInstancedPropertyBag(AssetBag);
 
 	// Next add shared variables
 	for (const UAnimNextRigVMAsset* ReferencedVariableAsset : Asset->ReferencedVariableAssets)
@@ -2746,7 +2755,9 @@ FInstancedPropertyBag UAnimNextRigVMAssetEditorData::GenerateCombinedPropertyBag
 			continue;
 		}
 
-		AddInstancedPropertyBag(ReferencedVariableAsset->GetVariableDefaults());
+		const UAnimNextRigVMAssetEditorData* ReferencedEditorData = UE::AnimNext::UncookedOnly::FUtils::GetEditorData<UAnimNextRigVMAssetEditorData>(ReferencedVariableAsset);
+		AssetBag = UE::AnimNext::UncookedOnly::FUtils::MakePropertyBagForEditorData(ReferencedEditorData, InCompileContext);
+		AddInstancedPropertyBag(AssetBag);
 	}
 
 	// Next add native structs
diff --git a/Engine/Plugins/Experimental/AnimNext/Source/AnimNextUncookedOnly/Private/AnimNextScopedCompileJob.cpp b/Engine/Plugins/Experimental/AnimNext/Source/AnimNextUncookedOnly/Private/AnimNextScopedCompileJob.cpp
new file mode 100644
--- /dev/null
+++ b/Engine/Plugins/Experimental/AnimNext/Source/AnimNextUncookedOnly/Private/AnimNextScopedCompileJob.cpp
@@ -0,0 +1,171 @@
+﻿// Copyright Epic Games, Inc. All Rights Reserved.
+
+#include "AnimNextScopedCompileJob.h"
+
+#include "AnimNextRigVMAsset.h"
+#include "AnimNextRigVMAssetEditorData.h"
+#include "UncookedOnlyUtils.h"
+#include "Kismet2/CompilerResultsLog.h"
+#include "Logging/MessageLog.h"
+
+#define LOCTEXT_NAMESPACE "AnimNextScopedCompileJob"
+
+namespace UE::AnimNext::UncookedOnly
+{
+
+struct FCompileJobThreadData
+{
+	TArray<TWeakPtr<FCompilerResultsLog>> CurrentLogStack;
+	TArray<TSharedRef<FTokenizedMessage>> Messages;
+	TSet<TWeakObjectPtr<UAnimNextRigVMAsset>> CompiledAssets;
+	int32 NumErrors = 0;
+	int32 NumWarnings = 0;
+	bool bDidCompile = false;
+};
+
+static thread_local FCompileJobThreadData GCompileJobData;
+
+FScopedCompileJob::FScopedCompileJob(const FText& InJobName, TConstArrayView<UObject*> InAssets)
+	: FScopedCompileJob(InJobName, nullptr, InAssets)
+{
+}
+
+FScopedCompileJob::FScopedCompileJob(UObject* InObject)
+	: FScopedCompileJob(FText::FromName(InObject->GetFName()), InObject, TArrayView<UObject*>(&InObject, 1))
+{
+}
+
+FScopedCompileJob::FScopedCompileJob(const FText& InJobName, UObject* InObject, TConstArrayView<UObject*> InAssets)
+	: JobName(InJobName)
+	, Object(InObject)
+{
+	FCompileJobThreadData* ThreadData = &GCompileJobData;
+	StartTime = FPlatformTime::Seconds();
+	ThreadData->bDidCompile = Object != nullptr;
+	Log = MakeShared<FCompilerResultsLog>(); 
+	ThreadData->CurrentLogStack.Push(Log);
+
+	for(UObject* Asset : InAssets)
+	{
+		UAnimNextRigVMAsset* AnimNextRigVMAsset = Cast<UAnimNextRigVMAsset>(Asset);
+		if(AnimNextRigVMAsset == nullptr)
+		{
+			return;
+		}
+
+		UAnimNextRigVMAssetEditorData* EditorData = FUtils::GetEditorData<UAnimNextRigVMAssetEditorData>(AnimNextRigVMAsset);
+		if(EditorData == nullptr)
+		{
+			return;
+		}
+
+		bool bAlreadyCompiling = false;
+		ThreadData->CompiledAssets.Add(AnimNextRigVMAsset, &bAlreadyCompiling);
+
+		if (!bAlreadyCompiling)
+		{
+			EditorData->ClearErrorInfoForAllEdGraphs();
+		}
+
+		if (!bAlreadyCompiling && !EditorData->bSuspendCompilationNotifications)
+		{
+			EditorData->OnCompileJobStarted();
+			UAnimNextRigVMAsset::OnCompileJobStarted().Broadcast(AnimNextRigVMAsset);
+		}
+	}
+}
+
+FScopedCompileJob::~FScopedCompileJob()
+{
+	FinishTime = FPlatformTime::Seconds();
+
+	FCompileJobThreadData* ThreadData = &GCompileJobData;
+	
+	// Accumulate messages
+	ThreadData->Messages.Append(Log->Messages);
+	ThreadData->NumErrors += Log->NumErrors;
+	ThreadData->NumWarnings += Log->NumWarnings;
+
+	ThreadData->CurrentLogStack.Pop();
+
+	if(ThreadData->CurrentLogStack.IsEmpty())
+	{
+		if(ThreadData->bDidCompile)
+		{
+			FMessageLog MessageLog("AnimNextCompilerResults");
+
+			MessageLog.NewPage(FText::Format(LOCTEXT("CompileFormat", "Compile {0}: {1}"), JobName, FText::AsDateTime(FDateTime::UtcNow())));
+			MessageLog.AddMessages(ThreadData->Messages);
+
+			// Print summary
+			FNumberFormattingOptions TimeFormat;
+			TimeFormat.MaximumFractionalDigits = 2;
+			TimeFormat.MinimumFractionalDigits = 2;
+			TimeFormat.MaximumIntegralDigits = 4;
+			TimeFormat.MinimumIntegralDigits = 4;
+			TimeFormat.UseGrouping = false;
+
+			FFormatNamedArguments Args;
+			Args.Add(TEXT("CurrentTime"), FText::AsNumber(FinishTime - GStartTime, &TimeFormat));
+			Args.Add(TEXT("JobName"), JobName);
+			Args.Add(TEXT("CompileTime"), (int32)((FinishTime - StartTime) * 1000));
+			Args.Add(TEXT("ObjectPath"), Object != nullptr ? FText::Format(LOCTEXT("ObjectPathFormat", "({0})"), FText::FromString(Object->GetPathName())) : FText::GetEmpty());
+			Args.Add(TEXT("NumAssets"), ThreadData->CompiledAssets.Num());
+
+			if (ThreadData->NumErrors > 0)
+			{
+				Args.Add(TEXT("NumErrors"), ThreadData->NumErrors);
+				Args.Add(TEXT("NumWarnings"), ThreadData->NumWarnings);
+				MessageLog.Info(FText::Format(LOCTEXT("CompileFailed", "[{CurrentTime}] Compile of {JobName} failed. {NumErrors} Error(s), {NumWarnings} Warning(s) [{NumAssets} assets in {CompileTime} ms] {ObjectPath}"), MoveTemp(Args)));
+			}
+			else if(ThreadData->NumWarnings > 0)
+			{
+				Args.Add(TEXT("NumWarnings"), ThreadData->NumWarnings);
+				MessageLog.Info(FText::Format(LOCTEXT("CompileWarning", "[{CurrentTime}] Compile of {JobName} successful. {NumWarnings} Warning(s) [{NumAssets} assets in {CompileTime} ms] {ObjectPath}"), MoveTemp(Args)));
+			}
+			else
+			{
+				MessageLog.Info(FText::Format(LOCTEXT("CompileSuccess", "[{CurrentTime}] Compile of {JobName} successful! [{NumAssets} assets in {CompileTime} ms] {ObjectPath}"), MoveTemp(Args)));
+			}
+
+			// Broadcast compilation finished so reallocation can occur
+			for (TWeakObjectPtr<UAnimNextRigVMAsset> WeakAsset : ThreadData->CompiledAssets)
+			{
+				UAnimNextRigVMAsset* Asset = WeakAsset.Get();
+				if (Asset == nullptr)
+				{
+					continue;
+				}
+
+				UAnimNextRigVMAssetEditorData* EditorData = FUtils::GetEditorData<UAnimNextRigVMAssetEditorData>(Asset);
+				if (EditorData == nullptr)
+				{
+					continue;
+				}
+
+				if (!EditorData->bSuspendCompilationNotifications)
+				{
+					UAnimNextRigVMAsset::OnCompileJobStarted().Broadcast(Asset);
+					EditorData->OnCompileJobFinished();
+				}
+			}
+		}
+
+		ThreadData->CompiledAssets.Empty();
+		ThreadData->Messages.Empty();
+		ThreadData->NumErrors = 0;
+		ThreadData->NumWarnings = 0;
+		ThreadData->bDidCompile = false;
+	}
+}
+
+FCompilerResultsLog& FScopedCompileJob::GetLog()
+{
+	FCompileJobThreadData* ThreadData = &GCompileJobData;
+	check(ThreadData->CurrentLogStack.Num() && ThreadData->CurrentLogStack.Top().IsValid());
+	return *ThreadData->CurrentLogStack.Top().Pin().Get();
+}
+
+}
+
+#undef LOCTEXT_NAMESPACE
\ No newline at end of file
diff --git a/Engine/Plugins/Experimental/AnimNext/Source/AnimNextUncookedOnly/Private/AnimNextScopedCompilerResults.cpp b/Engine/Plugins/Experimental/AnimNext/Source/AnimNextUncookedOnly/Private/AnimNextScopedCompilerResults.cpp
deleted file mode 100644
diff --git a/Engine/Plugins/Experimental/AnimNext/Source/AnimNextUncookedOnly/Private/Entries/AnimNextRigVMAssetEntry.cpp b/Engine/Plugins/Experimental/AnimNext/Source/AnimNextUncookedOnly/Private/Entries/AnimNextRigVMAssetEntry.cpp
--- a/Engine/Plugins/Experimental/AnimNext/Source/AnimNextUncookedOnly/Private/Entries/AnimNextRigVMAssetEntry.cpp
+++ b/Engine/Plugins/Experimental/AnimNext/Source/AnimNextUncookedOnly/Private/Entries/AnimNextRigVMAssetEntry.cpp
@@ -2,7 +2,6 @@
 
 #include "Entries/AnimNextRigVMAssetEntry.h"
 #include "AnimNextRigVMAssetEditorData.h"
-#include "AnimNextScopedCompilerResults.h"
 #include "Misc/TransactionObjectEvent.h"
 
 #define LOCTEXT_NAMESPACE "AnimNextRigVMAssetEntry"
@@ -39,15 +38,19 @@ void UAnimNextRigVMAssetEntry::PostTransacted(const FTransactionObjectEvent& Tra
 	}
 }
 
+void UAnimNextRigVMAssetEntry::PostEditUndo()
+{
+	// Note: Dont call super as that will default to a generic property changed event
+
+	BroadcastModified(EAnimNextEditorDataNotifType::UndoRedo);
+}
+
 #endif
 
 void UAnimNextRigVMAssetEntry::BroadcastModified(EAnimNextEditorDataNotifType InType)
 {
 	if(UAnimNextRigVMAssetEditorData* EditorData = Cast<UAnimNextRigVMAssetEditorData>(GetOuter()))
 	{
-		// Modifications here can trigger compilation, so add a scope to catch compile batches
-		UE::AnimNext::UncookedOnly::FScopedCompilerResults CompilerResults(LOCTEXT("ModifiedAssetsJobName", "Modified Assets"));
-
 		EditorData->BroadcastModified(InType, this);
 	}
 }
diff --git a/Engine/Plugins/Experimental/AnimNext/Source/AnimNextUncookedOnly/Private/Module/AnimNextModule_EditorData.cpp b/Engine/Plugins/Experimental/AnimNext/Source/AnimNextUncookedOnly/Private/Module/AnimNextModule_EditorData.cpp
--- a/Engine/Plugins/Experimental/AnimNext/Source/AnimNextUncookedOnly/Private/Module/AnimNextModule_EditorData.cpp
+++ b/Engine/Plugins/Experimental/AnimNext/Source/AnimNextUncookedOnly/Private/Module/AnimNextModule_EditorData.cpp
@@ -15,23 +15,6 @@
 #include "UObject/LinkerLoad.h"
 #include "Variables/AnimNextUniversalObjectLocatorBindingData.h"
 
-void UAnimNextModule_EditorData::RecompileVM()
-{
-	Super::RecompileVM();
-
-	if (bIsCompiling)
-	{
-		return;
-	}
-
-#if WITH_EDITOR
-	if(!bSuspendCompilationNotifications)
-	{
-		UAnimNextModule::OnModuleCompiled().Broadcast(UE::AnimNext::UncookedOnly::FUtils::GetAsset<UAnimNextModule>(this));
-	}
-#endif
-}
-
 TConstArrayView<TSubclassOf<UAnimNextRigVMAssetEntry>> UAnimNextModule_EditorData::GetEntryClasses() const
 {
 	static const TSubclassOf<UAnimNextRigVMAssetEntry> Classes[] =
diff --git a/Engine/Plugins/Experimental/AnimNext/Source/AnimNextUncookedOnly/Private/UncookedOnlyUtils.cpp b/Engine/Plugins/Experimental/AnimNext/Source/AnimNextUncookedOnly/Private/UncookedOnlyUtils.cpp
--- a/Engine/Plugins/Experimental/AnimNext/Source/AnimNextUncookedOnly/Private/UncookedOnlyUtils.cpp
+++ b/Engine/Plugins/Experimental/AnimNext/Source/AnimNextUncookedOnly/Private/UncookedOnlyUtils.cpp
@@ -53,19 +53,10 @@
 	InAsset->RigVM = InAsset->VM; // Local serialization
 }
 
-
-void FUtils::CompileVariables(const FRigVMCompileSettings& InSettings, UAnimNextRigVMAsset* InAsset, FAnimNextGetVariableCompileContext& OutCompileContext)
+FInstancedPropertyBag FUtils::MakePropertyBagForEditorData(const UAnimNextRigVMAssetEditorData* InEditorData, const FAnimNextGetVariableCompileContext& InCompileContext)
 {
-	check(InAsset);
-	
 	FMessageLog Log("AnimNextCompilerResults");
-
-	UAnimNextRigVMAssetEditorData* EditorData = GetEditorData<UAnimNextRigVMAssetEditorData>(InAsset);
-
-	// Gather programmatic variables regenerated each compile
-	EditorData->OnPreCompileGetProgrammaticVariables(InSettings, OutCompileContext);
-	const TArray<FAnimNextProgrammaticVariable>& ProgrammaticVariables = OutCompileContext.GetProgrammaticVariables();
-
+	
 	struct FStructEntryInfo
 	{
 		FName Name;
@@ -78,12 +69,10 @@
 	// Gather all variables in this asset.
 	TMap<FName, int32> EntryInfoIndexMap;
 	TArray<FStructEntryInfo> StructEntryInfos;
-	StructEntryInfos.Reserve(EditorData->Entries.Num() + ProgrammaticVariables.Num());
+	const TArray<FAnimNextProgrammaticVariable>& ProgrammaticVariables = InCompileContext.GetProgrammaticVariables();
+	StructEntryInfos.Reserve(InEditorData->Entries.Num() + ProgrammaticVariables.Num());
 	int32 NumPublicVariables = 0;
 
-	TSet<const UAnimNextRigVMAsset*> SharedVariableAssets;
-	TSet<const UScriptStruct*> SharedVariableStructs;
-
 	auto AddVariable = [&Log, &NumPublicVariables, &StructEntryInfos, &EntryInfoIndexMap](FName InName, const FAnimNextParamType& InType, EAnimNextExportAccessSpecifier InAccess, TConstArrayView<uint8> InValue, EPropertyFlags InFlags)
 	{
 		if(!InType.IsValid())
@@ -116,7 +105,7 @@
 		EntryInfoIndexMap.Add(InName, Index);
 	};
 
-	for(UAnimNextRigVMAssetEntry* Entry : EditorData->Entries)
+	for(UAnimNextRigVMAssetEntry* Entry : InEditorData->Entries)
 	{
 		if(const UAnimNextVariableEntry* VariableEntry = Cast<UAnimNextVariableEntry>(Entry))
 		{
@@ -137,9 +126,76 @@
 				TConstArrayView<uint8>(VariableEntry->GetValuePtr(), VariableEntry->GetType().GetSize()),
 				Flags);
 		}
-		else if(const UAnimNextSharedVariablesEntry* SharedVariablesEntry = Cast<UAnimNextSharedVariablesEntry>(Entry))
+	}
+
+	for (const FAnimNextProgrammaticVariable& Variable : ProgrammaticVariables)
+	{
+		AddVariable(
+			Variable.Name,
+			Variable.Type,
+			EAnimNextExportAccessSpecifier::Private,
+			TConstArrayView<uint8>(Variable.GetValuePtr(), Variable.Type.GetSize()),
+			CPF_AdvancedDisplay | CPF_NativeAccessSpecifierPrivate);
+	}
+
+	// Sort by size, largest first, for better packing
+	StructEntryInfos.Sort([](const FStructEntryInfo& InLHS, const FStructEntryInfo& InRHS)
+	{
+		if(InLHS.Type.GetSize() != InRHS.Type.GetSize())
 		{
-			
+			return InLHS.Type.GetSize() > InRHS.Type.GetSize();
+		}
+		else
+		{
+			return InLHS.Name.LexicalLess(InRHS.Name);
+		}
+	});
+
+	FInstancedPropertyBag VariableDefaults;
+
+	if(StructEntryInfos.Num() > 0)
+	{
+		// Build PropertyDescs and values to batch-create the property bag
+		TArray<FPropertyBagPropertyDesc> PropertyDescs;
+		PropertyDescs.Reserve(StructEntryInfos.Num());
+		TArray<TConstArrayView<uint8>> Values;
+		Values.Reserve(StructEntryInfos.Num());
+
+		for (const FStructEntryInfo& StructEntryInfo : StructEntryInfos)
+		{
+			PropertyDescs.Emplace(StructEntryInfo.Name, StructEntryInfo.Type.ContainerType, StructEntryInfo.Type.ValueType, StructEntryInfo.Type.ValueTypeObject, StructEntryInfo.PropertyFlags);
+			Values.Add(StructEntryInfo.Value);
+		}
+
+		// Create new property bags and migrate
+		EPropertyBagResult Result = VariableDefaults.ReplaceAllPropertiesAndValues(PropertyDescs, Values);
+		check(Result == EPropertyBagResult::Success);
+	}
+
+	return VariableDefaults;
+}
+
+void FUtils::CompileVariables(const FRigVMCompileSettings& InSettings, UAnimNextRigVMAsset* InAsset, FAnimNextGetVariableCompileContext& OutCompileContext)
+{
+	check(InAsset);
+
+	UAnimNextRigVMAssetEditorData* EditorData = GetEditorData<UAnimNextRigVMAssetEditorData>(InAsset);
+
+	// Gather programmatic variables regenerated each compile
+	EditorData->OnPreCompileGetProgrammaticVariables(InSettings, OutCompileContext);
+
+	// Generate the internal property bag
+	InAsset->VariableDefaults = MakePropertyBagForEditorData(EditorData, OutCompileContext);
+
+	TSet<const UAnimNextRigVMAsset*> SharedVariableAssets;
+	TSet<const UScriptStruct*> SharedVariableStructs;
+
+	FMessageLog Log("AnimNextCompilerResults");
+
+	for(UAnimNextRigVMAssetEntry* Entry : EditorData->Entries)
+	{
+		if(const UAnimNextSharedVariablesEntry* SharedVariablesEntry = Cast<UAnimNextSharedVariablesEntry>(Entry))
+		{
 			switch (SharedVariablesEntry->GetType())
 			{
 			case EAnimNextSharedVariablesType::Asset:
@@ -181,30 +237,7 @@
 			}
 		}
 	}
-
-	for (const FAnimNextProgrammaticVariable& Variable : ProgrammaticVariables)
-	{
-		AddVariable(
-			Variable.Name,
-			Variable.Type,
-			EAnimNextExportAccessSpecifier::Private,
-			TConstArrayView<uint8>(Variable.GetValuePtr(), Variable.Type.GetSize()),
-			CPF_AdvancedDisplay | CPF_NativeAccessSpecifierPrivate);
-	}
-
-	// Sort by size, largest first, for better packing
-	StructEntryInfos.Sort([](const FStructEntryInfo& InLHS, const FStructEntryInfo& InRHS)
-	{
-		if(InLHS.Type.GetSize() != InRHS.Type.GetSize())
-		{
-			return InLHS.Type.GetSize() > InRHS.Type.GetSize();
-		}
-		else
-		{
-			return InLHS.Name.LexicalLess(InRHS.Name);
-		}
-	});
-
+	
 	// Set shared variable assets/structs before we setup variables as GetExternalVariablesImpl relies on these
 	if (SharedVariableAssets.Num() > 0)
 	{
@@ -223,32 +256,9 @@
 	{
 		InAsset->ReferencedVariableStructs.Empty();
 	}
-	
-	if(StructEntryInfos.Num() > 0)
-	{
-		// Build PropertyDescs and values to batch-create the property bag
-		TArray<FPropertyBagPropertyDesc> PropertyDescs;
-		PropertyDescs.Reserve(StructEntryInfos.Num());
-		TArray<TConstArrayView<uint8>> Values;
-		Values.Reserve(StructEntryInfos.Num());
 
-		for (const FStructEntryInfo& StructEntryInfo : StructEntryInfos)
-		{
-			PropertyDescs.Emplace(StructEntryInfo.Name, StructEntryInfo.Type.ContainerType, StructEntryInfo.Type.ValueType, StructEntryInfo.Type.ValueTypeObject, StructEntryInfo.PropertyFlags);
-			Values.Add(StructEntryInfo.Value);
-		}
-
-		// Create new property bags and migrate
-		EPropertyBagResult Result = InAsset->VariableDefaults.ReplaceAllPropertiesAndValues(PropertyDescs, Values);
-		check(Result == EPropertyBagResult::Success);
-	}
-	else
-	{
-		InAsset->VariableDefaults.Reset();
-	}
-
 	// Now rebuild combined the combined property bag we used for stable properties 
-	InAsset->CombinedPropertyBag = EditorData->GenerateCombinedPropertyBag();
+	InAsset->CombinedPropertyBag = EditorData->GenerateCombinedPropertyBag(OutCompileContext);
 
 	// Rebuild external variables
 	if (InAsset->CombinedPropertyBag.GetNumPropertiesInBag() > 0)
diff --git a/Engine/Plugins/Experimental/AnimNext/Source/AnimNextUncookedOnly/Public/AnimNextScopedCompileJob.h b/Engine/Plugins/Experimental/AnimNext/Source/AnimNextUncookedOnly/Public/AnimNextScopedCompileJob.h
new file mode 100644
--- /dev/null
+++ b/Engine/Plugins/Experimental/AnimNext/Source/AnimNextUncookedOnly/Public/AnimNextScopedCompileJob.h
@@ -0,0 +1,37 @@
+﻿// Copyright Epic Games, Inc. All Rights Reserved.
+
+#pragma once
+
+#include "CoreMinimal.h"
+#include "Kismet2/CompilerResultsLog.h"
+
+class UAnimNextRigVMAssetEditorData;
+class FCompilerResultsLog;
+
+namespace UE::AnimNext::UncookedOnly
+{
+
+// RAII helper for scoping compilation batches
+// The re-allocation of compiled assets and their dependencies is deferred until the outermost scope exits.
+class FScopedCompileJob
+{
+public:
+	ANIMNEXTUNCOOKEDONLY_API explicit FScopedCompileJob(const FText& InJobName, TConstArrayView<UObject*> InAssets);
+	ANIMNEXTUNCOOKEDONLY_API explicit FScopedCompileJob(UObject* InObject);
+	ANIMNEXTUNCOOKEDONLY_API ~FScopedCompileJob();
+
+	// Get the log that is currently in scope, asserts if no scope is active
+	ANIMNEXTUNCOOKEDONLY_API static FCompilerResultsLog& GetLog();
+
+private:
+	ANIMNEXTUNCOOKEDONLY_API FScopedCompileJob(const FText& InJobName, UObject* InObject, TConstArrayView<UObject*> InAssets);
+
+private:
+	TSharedPtr<FCompilerResultsLog> Log;
+	FText JobName;
+	UObject* Object = nullptr;
+	double StartTime = 0.0;
+	double FinishTime = 0.0;
+};
+
+}
diff --git a/Engine/Plugins/Experimental/AnimNext/Source/AnimNextUncookedOnly/Public/AnimNextScopedCompilerResults.h b/Engine/Plugins/Experimental/AnimNext/Source/AnimNextUncookedOnly/Public/AnimNextScopedCompilerResults.h
deleted file mode 100644
diff --git a/Engine/Plugins/Experimental/AnimNextAnimGraph/Source/AnimNextAnimGraph/Private/Graph/AnimNextGraphInstance.cpp b/Engine/Plugins/Experimental/AnimNextAnimGraph/Source/AnimNextAnimGraph/Private/Graph/AnimNextGraphInstance.cpp
--- a/Engine/Plugins/Experimental/AnimNextAnimGraph/Source/AnimNextAnimGraph/Private/Graph/AnimNextGraphInstance.cpp
+++ b/Engine/Plugins/Experimental/AnimNextAnimGraph/Source/AnimNextAnimGraph/Private/Graph/AnimNextGraphInstance.cpp
@@ -19,9 +19,6 @@ DEFINE_STAT(STAT_AnimNext_Graph_RigVM);
 
 FAnimNextGraphInstance::FAnimNextGraphInstance()
 {
-#if WITH_EDITORONLY_DATA
-	UAnimNextModule::OnModuleCompiled().AddRaw(this, &FAnimNextGraphInstance::OnModuleCompiled);
-#endif
 }
 
 FAnimNextGraphInstance::~FAnimNextGraphInstance()
@@ -32,8 +29,6 @@ FAnimNextGraphInstance::~FAnimNextGraphInstance()
 void FAnimNextGraphInstance::Release()
 {
 #if WITH_EDITORONLY_DATA
-	UAnimNextModule::OnModuleCompiled().RemoveAll(this);
-
 	if(const UAnimNextAnimationGraph* Graph = GetAnimationGraph())
 	{
 		FScopeLock Lock(&Graph->GraphInstancesLock);
@@ -206,7 +201,14 @@ void FAnimNextGraphInstance::Thaw()
 {
 	if (const UAnimNextAnimationGraph* AnimationGraph = GetAnimationGraph())
 	{
-		MigrateVariables();
+		if (Variables.bHasBeenInitialized)
+		{
+			MigrateVariables();
+		}
+		else
+		{
+			InitializeVariables();
+		}
 
 		{
 			UE::AnimNext::FExecutionContext Context(*this);
@@ -224,17 +226,4 @@ void FAnimNextGraphInstance::Thaw()
 	}
 }
 
-void FAnimNextGraphInstance::OnModuleCompiled(UAnimNextModule* InModule)
-{
-	// If we are hosted transitively by the compiled module, or if we could be bound to the compiled module,
-	// migrate our variables (which will re-bind to outer hosts)
-	if(ModuleInstance)
-	{
-		if(InModule == ModuleInstance->GetModule())
-		{
-			MigrateVariables();
-		}
-	}
-}
-
 #endif
diff --git a/Engine/Plugins/Experimental/AnimNextAnimGraph/Source/AnimNextAnimGraph/Public/Graph/AnimNextGraphInstance.h b/Engine/Plugins/Experimental/AnimNextAnimGraph/Source/AnimNextAnimGraph/Public/Graph/AnimNextGraphInstance.h
--- a/Engine/Plugins/Experimental/AnimNextAnimGraph/Source/AnimNextAnimGraph/Public/Graph/AnimNextGraphInstance.h
+++ b/Engine/Plugins/Experimental/AnimNextAnimGraph/Source/AnimNextAnimGraph/Public/Graph/AnimNextGraphInstance.h
@@ -119,9 +119,6 @@ private:
 
 	// During graph compilation, once compilation is done we thaw existing graph instances to reallocate their memory
 	UE_API void Thaw();
-
-	// Hook into module compilation
-	UE_API void OnModuleCompiled(UAnimNextModule* InModule);
 #endif
 
 	// The entry point in Graph that this instance corresponds to 
diff --git a/Engine/Plugins/Experimental/AnimNextAnimGraph/Source/AnimNextAnimGraphUncookedOnly/Internal/Graph/AnimNextAnimationGraph_EditorData.h b/Engine/Plugins/Experimental/AnimNextAnimGraph/Source/AnimNextAnimGraphUncookedOnly/Internal/Graph/AnimNextAnimationGraph_EditorData.h
--- a/Engine/Plugins/Experimental/AnimNextAnimGraph/Source/AnimNextAnimGraphUncookedOnly/Internal/Graph/AnimNextAnimationGraph_EditorData.h
+++ b/Engine/Plugins/Experimental/AnimNextAnimGraph/Source/AnimNextAnimGraphUncookedOnly/Internal/Graph/AnimNextAnimationGraph_EditorData.h
@@ -73,10 +73,11 @@ protected:
 	UE_API virtual bool CanAddNewEntry(TSubclassOf<UAnimNextRigVMAssetEntry> InClass) const override;
 	UE_API virtual TSubclassOf<UAssetUserData> GetAssetUserDataClass() const override;
 	UE_API virtual void InitializeAssetUserData() override;
+	UE_API virtual void OnCompileJobStarted() override;
 	UE_API virtual void OnPreCompileAsset(FRigVMCompileSettings& InSettings) override;
 	UE_API virtual void OnPreCompileGetProgrammaticFunctionHeaders(const FRigVMCompileSettings& InSettings, FAnimNextGetFunctionHeaderCompileContext& OutCompileContext) override;
 	UE_API virtual void OnPreCompileProcessGraphs(const FRigVMCompileSettings& InSettings, FAnimNextProcessGraphCompileContext& OutCompileContext) override;
-	UE_API virtual void OnPostCompileCleanup(const FRigVMCompileSettings& InSettings) override;
+	UE_API virtual void OnCompileJobFinished() override;
 
 	// Super overrides
 	UE_API virtual void GetAnimNextAssetRegistryTags(FAssetRegistryTagsContext& Context, FAnimNextAssetRegistryExports& OutExports) const override;
diff --git a/Engine/Plugins/Experimental/AnimNextAnimGraph/Source/AnimNextAnimGraphUncookedOnly/Private/Graph/AnimNextAnimationGraph_EditorData.cpp b/Engine/Plugins/Experimental/AnimNextAnimGraph/Source/AnimNextAnimGraphUncookedOnly/Private/Graph/AnimNextAnimationGraph_EditorData.cpp
--- a/Engine/Plugins/Experimental/AnimNextAnimGraph/Source/AnimNextAnimGraphUncookedOnly/Private/Graph/AnimNextAnimationGraph_EditorData.cpp
+++ b/Engine/Plugins/Experimental/AnimNextAnimGraph/Source/AnimNextAnimGraphUncookedOnly/Private/Graph/AnimNextAnimationGraph_EditorData.cpp
@@ -3,7 +3,7 @@
 #include "Graph/AnimNextAnimationGraph_EditorData.h"
 #include "AnimGraphUncookedOnlyUtils.h"
 #include "AnimNextAnimGraphWorkspaceAssetUserData.h"
-#include "AnimNextScopedCompilerResults.h"
+#include "AnimNextScopedCompileJob.h"
 #include "RigVMPythonUtils.h"
 #include "Compilation/AnimNextGetFunctionHeaderCompileContext.h"
 #include "Compilation/AnimNextGetVariableCompileContext.h"
@@ -505,10 +505,6 @@ void UAnimNextAnimationGraph_EditorData::OnPreCompileAsset(FRigVMCompileSettings
 
 	UAnimNextAnimationGraph* AnimationGraph = FUtils::GetAsset<UAnimNextAnimationGraph>(this);
 
-	// Before we re-compile a graph, we need to release and live instances since we need the metadata we are about to replace
-	// to call trait destructors etc
-	AnimationGraph->FreezeGraphInstances();
-
 	AnimationGraph->EntryPoints.Empty();
 	AnimationGraph->ResolvedRootTraitHandles.Empty();
 	AnimationGraph->ResolvedEntryPoints.Empty();
@@ -599,7 +595,7 @@ void UAnimNextAnimationGraph_EditorData::OnPreCompileProcessGraphs(const FRigVMC
 			TempGraph->SetExternalPackage(nullptr);
 			for(URigVMNode* SourceNode : SourceGraph->GetNodes())
 			{
-				FScopedCompilerResults::GetLog().NotifyIntermediateObjectCreation(CreatedObjects.FindChecked(SourceNode), SourceNode);
+				FScopedCompileJob::GetLog().NotifyIntermediateObjectCreation(CreatedObjects.FindChecked(SourceNode), SourceNode);
 			}
 
 			UAnimNextController* TempController = CastChecked<UAnimNextController>(VMClient->GetOrCreateController(TempGraph));
@@ -712,7 +708,18 @@ void UAnimNextAnimationGraph_EditorData::OnPreCompileProcessGraphs(const FRigVMC
 	InOutGraphs.Append(NonAnimGraphs);
 }
 
-void UAnimNextAnimationGraph_EditorData::OnPostCompileCleanup(const FRigVMCompileSettings& InSettings)
+void UAnimNextAnimationGraph_EditorData::OnCompileJobStarted()
+{
+	using namespace UE::AnimNext::UncookedOnly;
+
+	UAnimNextAnimationGraph* AnimationGraph = FUtils::GetAsset<UAnimNextAnimationGraph>(this);
+
+	// Before we re-compile a graph, we need to release and live instances since we need the metadata we are about to replace
+	// to call trait destructors etc
+	AnimationGraph->FreezeGraphInstances();
+}
+
+void UAnimNextAnimationGraph_EditorData::OnCompileJobFinished()
 {
 	using namespace UE::AnimNext::UncookedOnly;
 
diff --git a/Engine/Plugins/Experimental/AnimNextStateTree/Source/AnimNextStateTreeEditor/Private/StateTreeAssetCompilationHandler.cpp b/Engine/Plugins/Experimental/AnimNextStateTree/Source/AnimNextStateTreeEditor/Private/StateTreeAssetCompilationHandler.cpp
--- a/Engine/Plugins/Experimental/AnimNextStateTree/Source/AnimNextStateTreeEditor/Private/StateTreeAssetCompilationHandler.cpp
+++ b/Engine/Plugins/Experimental/AnimNextStateTree/Source/AnimNextStateTreeEditor/Private/StateTreeAssetCompilationHandler.cpp
@@ -3,7 +3,7 @@
 #include "StateTreeAssetCompilationHandler.h"
 
 #include "AnimNextRigVMAssetEditorData.h"
-#include "AnimNextScopedCompilerResults.h"
+#include "AnimNextScopedCompileJob.h"
 #include "AnimNextStateTree.h"
 #include "AnimNextStateTree_EditorData.h"
 #include "Editor.h"
@@ -51,7 +51,7 @@ void FStateTreeAssetCompilationHandler::Compile(TSharedRef<Workspace::IWorkspace
 		return;
 	}
 
-	UE::AnimNext::UncookedOnly::FScopedCompilerResults CompilerResults(AnimNextStateTree);
+	UE::AnimNext::UncookedOnly::FScopedCompileJob CompileJob(AnimNextStateTree);
 
 	// We can't recompile state tree while in PIE
 	if (!GEditor->IsPlaySessionInProgress())
@@ -64,12 +64,12 @@ void FStateTreeAssetCompilationHandler::Compile(TSharedRef<Workspace::IWorkspace
 		// Add to messages
 		for(const TSharedRef<FTokenizedMessage>& Message : Log.ToTokenizedMessages())
 		{
-			CompilerResults.GetLog().AddTokenizedMessage(Message);
+			CompileJob.GetLog().AddTokenizedMessage(Message);
 		}
 	}
 	else
 	{
-		CompilerResults.GetLog().Note(*FString::Printf(TEXT("StateTree '%s' compilation was skipped because PIE was running"), *InAsset->GetName()));
+		CompileJob.GetLog().Note(*FString::Printf(TEXT("StateTree '%s' compilation was skipped because PIE was running"), *InAsset->GetName()));
 	}
 
 	// Call super to compile the RigVM Asset 
diff --git a/Engine/Plugins/Experimental/AnimNextStateTree/Source/AnimNextStateTreeUncookedOnly/Internal/AnimNextStateTree_EditorData.cpp b/Engine/Plugins/Experimental/AnimNextStateTree/Source/AnimNextStateTreeUncookedOnly/Internal/AnimNextStateTree_EditorData.cpp
--- a/Engine/Plugins/Experimental/AnimNextStateTree/Source/AnimNextStateTreeUncookedOnly/Internal/AnimNextStateTree_EditorData.cpp
+++ b/Engine/Plugins/Experimental/AnimNextStateTree/Source/AnimNextStateTreeUncookedOnly/Internal/AnimNextStateTree_EditorData.cpp
@@ -40,9 +40,6 @@ void UAnimNextStateTree_EditorData::RecompileVM()
 {
 	UAnimNextAnimationGraph_EditorData::RecompileVM();
 
-	// Refresh the 'dummy' property bag we have with all combined variables now we have compiled
-	RefreshCombinedVariablesPropertyBag();
-
 	// Skip further state tree compilation if we had errors in regular asset compilation. It can create spurious errors
 	if (bErrorsDuringCompilation)
 	{
@@ -221,7 +218,9 @@ void UAnimNextStateTree_EditorData::OnPreCompileGetProgrammaticGraphs(const FRig
 	}
 }
 
-void UAnimNextStateTree_EditorData::RefreshCombinedVariablesPropertyBag()
+void UAnimNextStateTree_EditorData::OnPostCompileVariables(const FRigVMCompileSettings& InSettings, const FAnimNextGetVariableCompileContext& InCompileContext)
 {
-	CombinedVariablesPropertyBag = GenerateCombinedPropertyBag();
-}
\ No newline at end of file
+	Super::OnPostCompileVariables(InSettings, InCompileContext);
+
+	CombinedVariablesPropertyBag = GenerateCombinedPropertyBag(InCompileContext);
+}
diff --git a/Engine/Plugins/Experimental/AnimNextStateTree/Source/AnimNextStateTreeUncookedOnly/Internal/AnimNextStateTree_EditorData.h b/Engine/Plugins/Experimental/AnimNextStateTree/Source/AnimNextStateTreeUncookedOnly/Internal/AnimNextStateTree_EditorData.h
--- a/Engine/Plugins/Experimental/AnimNextStateTree/Source/AnimNextStateTreeUncookedOnly/Internal/AnimNextStateTree_EditorData.h
+++ b/Engine/Plugins/Experimental/AnimNextStateTree/Source/AnimNextStateTreeUncookedOnly/Internal/AnimNextStateTree_EditorData.h
@@ -22,14 +22,9 @@ protected:
 	
 	// UAnimNextRigVMAssetEditorData interface
 	virtual TConstArrayView<TSubclassOf<UAnimNextRigVMAssetEntry>> GetEntryClasses() const override;
-
-	// Allows this asset to generate variables to be injected at compilation time, separate method to allow programmatic graphs to use these vars
 	virtual void OnPreCompileGetProgrammaticFunctionHeaders(const FRigVMCompileSettings& InSettings, FAnimNextGetFunctionHeaderCompileContext& OutCompileContext) override;
-
-	// Allows this asset to generate graphs to be injected at compilation time
 	virtual void OnPreCompileGetProgrammaticGraphs(const FRigVMCompileSettings& InSettings, FAnimNextGetGraphCompileContext& OutCompileContext) override;
-
-	void RefreshCombinedVariablesPropertyBag();
+	virtual void OnPostCompileVariables(const FRigVMCompileSettings& InSettings, const FAnimNextGetVariableCompileContext& InCompileContext) override;
 
 private:
 	// Property bag that forms the root parameters. Consists of all combined variables from shared variable assets/structs.
