UAF: Fix variable set node issues

Fix node validation generating suprious errors depending on load order. Use of IsValid relies on the target object being loaded and having valid variables. Checking IsNone and the object ptr manually allows the object to be preloaded before validation.
Fix type checking to accept types that can be converted sucecssfully (e.g. derived classes plugged into a base class type)

diff --git a/Engine/Plugins/Experimental/AnimNext/Source/AnimNextUncookedOnly/Private/Variables/K2Node_AnimNextComponentSetVariable.cpp b/Engine/Plugins/Experimental/AnimNext/Source/AnimNextUncookedOnly/Private/Variables/K2Node_AnimNextComponentSetVariable.cpp
--- a/Engine/Plugins/Experimental/AnimNext/Source/AnimNextUncookedOnly/Private/Variables/K2Node_AnimNextComponentSetVariable.cpp
+++ b/Engine/Plugins/Experimental/AnimNext/Source/AnimNextUncookedOnly/Private/Variables/K2Node_AnimNextComponentSetVariable.cpp
@@ -10,6 +10,8 @@
 #include "UncookedOnlyUtils.h"
 #include "Kismet2/BlueprintEditorUtils.h"
 #include "Component/AnimNextComponent.h"
+#include "Param/ParamCompatibility.h"
+#include "Param/ParamUtils.h"
 #include "Variables/AnimNextVariableReference.h"
 
 #define LOCTEXT_NAMESPACE "K2Node_AnimNextComponentSetVariable"
@@ -95,9 +97,10 @@ void UK2Node_AnimNextComponentSetVariable::ValidateNodeDuringCompilation(FCompil
 	FAnimNextVariableReference::StaticStruct()->ImportText(*VariablePin->DefaultValue, &VariableReference, nullptr, PPF_None, nullptr, FAnimNextVariableReference::StaticStruct()->GetName());
 	if (VariableReference.IsValid())
 	{
-		FAnimNextParamType Type = UE::AnimNext::UncookedOnly::FUtils::FindVariableType(VariableReference);
-		FAnimNextParamType PinType = UE::AnimNext::UncookedOnly::FUtils::GetParamTypeFromPinType(ValuePin->PinType);
-		if (PinType != Type)
+		FAnimNextParamType Type = UE::AnimNext::UncookedOnly::FUtils::FindVariableType(VariableReference);				// The variable from the UAF Module
+		FAnimNextParamType PinType = UE::AnimNext::UncookedOnly::FUtils::GetParamTypeFromPinType(ValuePin->PinType);	// The variable from BP pin
+		UE::AnimNext::FParamCompatibility Compatibility = UE::AnimNext::FParamUtils::GetCompatibility(Type, PinType);
+		if (Compatibility == UE::AnimNext::EParamCompatibility::Incompatible || Compatibility == UE::AnimNext::EParamCompatibility::Incompatible_DataLoss)
 		{
 			MessageLog.Error(*FString::Printf(TEXT("@@ mismatched type, setting '%s' type, but variable is '%s' type"), *PinType.ToString(), *Type.ToString()), this);
 		}
@@ -131,7 +134,7 @@ void UK2Node_AnimNextComponentSetVariable::PreloadRequiredAssets()
 	UEdGraphPin* VariablePin = FindPinChecked(TEXT("Variable"));
 	FAnimNextVariableReference VariableReference;
 	FAnimNextVariableReference::StaticStruct()->ImportText(*VariablePin->DefaultValue, &VariableReference, nullptr, PPF_None, nullptr, FAnimNextVariableReference::StaticStruct()->GetName());
-	if (VariableReference.IsValid())
+	if (!VariableReference.IsNone() && VariableReference.GetObject())
 	{
 		UObject* Object = const_cast<UObject*>(VariableReference.GetObject().Get());
 		PreloadObject(Object);
