UAF: Fix crash when tracing AnimNext data

diff --git a/Engine/Plugins/Experimental/AnimNext/Source/AnimNext/Private/RewindDebugger/AnimNextTrace.cpp b/Engine/Plugins/Experimental/AnimNext/Source/AnimNext/Private/RewindDebugger/AnimNextTrace.cpp
--- a/Engine/Plugins/Experimental/AnimNext/Source/AnimNext/Private/RewindDebugger/AnimNextTrace.cpp
+++ b/Engine/Plugins/Experimental/AnimNext/Source/AnimNext/Private/RewindDebugger/AnimNextTrace.cpp
@@ -176,6 +176,11 @@ void FAnimNextTrace::OutputAnimNextVariableSet(const TSharedRef<FUAFInstanceVari
 			// TODO: Accomodate overrides
 
 			FInstancedPropertyBag& InstancedPropertyBag = InVariableSet->VariablesContainer.Get<FInstancedPropertyBag>();
+			if (InstancedPropertyBag.GetNumPropertiesInBag() == 0)
+			{
+				break;
+			}
+
 			TConstArrayView<FPropertyBagPropertyDesc> PropertyDescs = InstancedPropertyBag.GetPropertyBagStruct()->GetPropertyDescs();
 			uint32 PropertyDescHash = GetPropertyDescHash(PropertyDescs);
 			{
