UAF: Fix 'empty' assets when loaded via a BP reference in a SetVariable node

diff --git a/Engine/Plugins/Experimental/AnimNext/Source/AnimNext/Private/AnimNextRigVMAsset.cpp b/Engine/Plugins/Experimental/AnimNext/Source/AnimNext/Private/AnimNextRigVMAsset.cpp
--- a/Engine/Plugins/Experimental/AnimNext/Source/AnimNext/Private/AnimNextRigVMAsset.cpp
+++ b/Engine/Plugins/Experimental/AnimNext/Source/AnimNext/Private/AnimNextRigVMAsset.cpp
@@ -40,6 +40,12 @@ void UAnimNextRigVMAsset::PostLoad()
 {
 	Super::PostLoad();
 
+#if WITH_EDITORONLY_DATA
+	// External editor-side APIs assume that the editordata is also postloaded with the asset
+	EditorData->ConditionalPreload();
+	EditorData->ConditionalPostLoad();
+#endif
+	
 	ExtendedExecuteContext.InvalidateCachedMemory();
 
 	VM = RigVM;
diff --git a/Engine/Plugins/Experimental/AnimNext/Source/AnimNextUncookedOnly/Private/AnimNextRigVMAssetEditorData.cpp b/Engine/Plugins/Experimental/AnimNext/Source/AnimNextUncookedOnly/Private/AnimNextRigVMAssetEditorData.cpp
--- a/Engine/Plugins/Experimental/AnimNext/Source/AnimNextUncookedOnly/Private/AnimNextRigVMAssetEditorData.cpp
+++ b/Engine/Plugins/Experimental/AnimNext/Source/AnimNextUncookedOnly/Private/AnimNextRigVMAssetEditorData.cpp
@@ -215,6 +215,14 @@ void UAnimNextRigVMAssetEditorData::PostLoad()
 	GraphModels.Reset();
 	
 	PostLoadExternalPackages();
+
+	// Preload our entries, as we need them for RefreshExternalModels
+	for (UAnimNextRigVMAssetEntry* Entry : Entries)
+	{
+		Entry->ConditionalPreload();
+		Entry->ConditionalPostLoad();
+	}
+
 	RefreshExternalModels();
 
 	Initialize(/*bRecompileVM*/false);
diff --git a/Engine/Plugins/Experimental/AnimNext/Source/AnimNextUncookedOnly/Private/Variables/K2Node_AnimNextComponentSetVariable.cpp b/Engine/Plugins/Experimental/AnimNext/Source/AnimNextUncookedOnly/Private/Variables/K2Node_AnimNextComponentSetVariable.cpp
--- a/Engine/Plugins/Experimental/AnimNext/Source/AnimNextUncookedOnly/Private/Variables/K2Node_AnimNextComponentSetVariable.cpp
+++ b/Engine/Plugins/Experimental/AnimNext/Source/AnimNextUncookedOnly/Private/Variables/K2Node_AnimNextComponentSetVariable.cpp
@@ -139,13 +139,6 @@ void UK2Node_AnimNextComponentSetVariable::PreloadRequiredAssets()
 		UObject* Object = const_cast<UObject*>(VariableReference.GetObject().Get());
 		PreloadObject(Object);
 		Object->ConditionalPostLoad();
-		TArray<UObject*> SubObjects;
-		GetObjectsWithOuter(Object, SubObjects);
-		for (UObject* SubObject : SubObjects)
-		{
-			PreloadObject(SubObject);
-			SubObject->ConditionalPostLoad();
-		}
 	}
 }
 
