UAF: Fix Rewind Debugger tracing

Update logger/channel tags to be consistent

diff --git a/Engine/Plugins/Experimental/AnimNext/Source/AnimNext/Private/RewindDebugger/AnimNextTrace.cpp b/Engine/Plugins/Experimental/AnimNext/Source/AnimNext/Private/RewindDebugger/AnimNextTrace.cpp
--- a/Engine/Plugins/Experimental/AnimNext/Source/AnimNext/Private/RewindDebugger/AnimNextTrace.cpp
+++ b/Engine/Plugins/Experimental/AnimNext/Source/AnimNext/Private/RewindDebugger/AnimNextTrace.cpp
@@ -14,14 +14,14 @@
 #if ANIMNEXT_TRACE_ENABLED
 
 
-UE_TRACE_EVENT_BEGIN(AnimNext, Instance)
+UE_TRACE_EVENT_BEGIN(UAF, Instance)
 	UE_TRACE_EVENT_FIELD(uint64, InstanceId)
 	UE_TRACE_EVENT_FIELD(uint64, HostInstanceId)
 	UE_TRACE_EVENT_FIELD(uint64, OuterObjectId)
 	UE_TRACE_EVENT_FIELD(uint64, AssetId)
 UE_TRACE_EVENT_END()
 
-UE_TRACE_EVENT_BEGIN(AnimNext, InstanceVariables)
+UE_TRACE_EVENT_BEGIN(UAF, InstanceVariables)
 	UE_TRACE_EVENT_FIELD(uint64, Cycle)
 	UE_TRACE_EVENT_FIELD(double, RecordingTime)
 	UE_TRACE_EVENT_FIELD(uint64, InstanceId)
@@ -29,12 +29,12 @@ UE_TRACE_EVENT_BEGIN(AnimNext, InstanceVariables)
 	UE_TRACE_EVENT_FIELD(uint8[], VariableData)
 UE_TRACE_EVENT_END()
 
-UE_TRACE_EVENT_BEGIN(AnimNext, InstanceVariableDescriptions)
+UE_TRACE_EVENT_BEGIN(UAF, InstanceVariableDescriptions)
 	UE_TRACE_EVENT_FIELD(uint32, VariableDescriptionHash)
 	UE_TRACE_EVENT_FIELD(uint8[], VariableDescriptionData)
 UE_TRACE_EVENT_END()
 
-UE_TRACE_CHANNEL_DEFINE(AnimNextChannel);
+UE_TRACE_CHANNEL_DEFINE(UAFChannel);
 
 namespace UE::AnimNext
 {
@@ -70,7 +70,7 @@ void FAnimNextTrace::Reset()
 
 void FAnimNextTrace::OutputAnimNextInstance(const FUAFAssetInstance* AnimNextInstance, const UObject* OuterObject)
 {
-	bool bChannelEnabled = UE_TRACE_CHANNELEXPR_IS_ENABLED(AnimNextChannel);
+	bool bChannelEnabled = UE_TRACE_CHANNELEXPR_IS_ENABLED(UAFChannel);
 	if (!bChannelEnabled || AnimNextInstance == nullptr)
 	{
 		return;
@@ -110,7 +110,7 @@ void FAnimNextTrace::OutputAnimNextInstance(const FUAFAssetInstance* AnimNextIns
 		
 		const FUAFAssetInstance* HostInstance = AnimNextInstance->GetHost();
 
-		UE_TRACE_LOG(AnimNext, Instance, AnimNextChannel)
+		UE_TRACE_LOG(UAF, Instance, UAFChannel)
 			<< Instance.InstanceId(InstanceId)
 			<< Instance.OuterObjectId(FObjectTrace::GetObjectId(OuterObject))
 			<< Instance.HostInstanceId(HostInstance ? HostInstance->GetUniqueId() : 0)
@@ -136,7 +136,7 @@ uint32 GetPropertyDescHash(const TConstArrayView<FPropertyBagPropertyDesc>& Desc
 
 void FAnimNextTrace::OutputAnimNextVariables(const FUAFAssetInstance* Instance, const UObject* OuterObject)
 {
-	bool bChannelEnabled = UE_TRACE_CHANNELEXPR_IS_ENABLED(AnimNextChannel);
+	bool bChannelEnabled = UE_TRACE_CHANNELEXPR_IS_ENABLED(UAFChannel);
 	if (!bChannelEnabled || Instance == nullptr || OuterObject == nullptr)
 	{
 		return;
@@ -160,7 +160,7 @@ void FAnimNextTrace::OutputAnimNextVariables(const FUAFAssetInstance* Instance,
 	}
 	else
 	{
-		UE_TRACE_LOG(AnimNext, InstanceVariables, AnimNextChannel)
+		UE_TRACE_LOG(UAF, InstanceVariables, UAFChannel)
 				<< InstanceVariables.Cycle(FPlatformTime::Cycles64())
 				<< InstanceVariables.RecordingTime(FObjectTrace::GetWorldElapsedTime(OuterObject->GetWorld()))
 				<< InstanceVariables.InstanceId(InstanceId);
@@ -209,7 +209,7 @@ void FAnimNextTrace::OutputAnimNextVariableSet(const TSharedRef<FUAFInstanceVari
 						Archive.UsingCustomVersion(FPropertyBagCustomVersion::GUID);
 						Archive << PropertyDescriptions;
 
-						UE_TRACE_LOG(AnimNext, InstanceVariableDescriptions, AnimNextChannel)
+						UE_TRACE_LOG(UAF, InstanceVariableDescriptions, UAFChannel)
 							<< InstanceVariableDescriptions.VariableDescriptionHash(PropertyDescHash)
 							<< InstanceVariableDescriptions.VariableDescriptionData(ArchiveData.GetData(), ArchiveData.Num());
 					}
@@ -224,7 +224,7 @@ void FAnimNextTrace::OutputAnimNextVariableSet(const TSharedRef<FUAFInstanceVari
 					PropertyBag->SerializeItem(Archive, InstancedPropertyBag.GetMutableValue().GetMemory(), nullptr);
 				}
 
-				UE_TRACE_LOG(AnimNext, InstanceVariables, AnimNextChannel)
+				UE_TRACE_LOG(UAF, InstanceVariables, UAFChannel)
 						<< InstanceVariables.Cycle(FPlatformTime::Cycles64())
 						<< InstanceVariables.RecordingTime(FObjectTrace::GetWorldElapsedTime(OuterObject->GetWorld()))
 						<< InstanceVariables.InstanceId(InstanceId)
diff --git a/Engine/Plugins/Experimental/AnimNext/Source/AnimNext/Public/RewindDebugger/AnimNextTrace.h b/Engine/Plugins/Experimental/AnimNext/Source/AnimNext/Public/RewindDebugger/AnimNextTrace.h
--- a/Engine/Plugins/Experimental/AnimNext/Source/AnimNext/Public/RewindDebugger/AnimNextTrace.h
+++ b/Engine/Plugins/Experimental/AnimNext/Source/AnimNext/Public/RewindDebugger/AnimNextTrace.h
@@ -18,7 +18,7 @@
 
 #if ANIMNEXT_TRACE_ENABLED
 
-UE_TRACE_CHANNEL_EXTERN(AnimNextChannel, ANIMNEXT_API);
+UE_TRACE_CHANNEL_EXTERN(UAFChannel, ANIMNEXT_API);
 
 class UObject;
 struct FAnimNextModuleInstance;
diff --git a/Engine/Plugins/Experimental/AnimNext/Source/AnimNextEditor/Private/RewindDebugger/AnimNextAnalyzer.cpp b/Engine/Plugins/Experimental/AnimNext/Source/AnimNextEditor/Private/RewindDebugger/AnimNextAnalyzer.cpp
--- a/Engine/Plugins/Experimental/AnimNext/Source/AnimNextEditor/Private/RewindDebugger/AnimNextAnalyzer.cpp
+++ b/Engine/Plugins/Experimental/AnimNext/Source/AnimNextEditor/Private/RewindDebugger/AnimNextAnalyzer.cpp
@@ -22,9 +22,9 @@ void FAnimNextAnalyzer::OnAnalysisBegin(const FOnAnalysisContext& Context)
 {
 	auto& Builder = Context.InterfaceBuilder;
 
-	Builder.RouteEvent(RouteId_Module, "AnimNext", "Instance");
-	Builder.RouteEvent(RouteId_InstanceVariables, "AnimNext", "InstanceVariables");
-	Builder.RouteEvent(RouteId_InstanceVariableDescriptions, "AnimNext", "InstanceVariableDescriptions");
+	Builder.RouteEvent(RouteId_Module, "UAF", "Instance");
+	Builder.RouteEvent(RouteId_InstanceVariables, "UAF", "InstanceVariables");
+	Builder.RouteEvent(RouteId_InstanceVariableDescriptions, "UAF", "InstanceVariableDescriptions");
 }
 
 bool FAnimNextAnalyzer::OnEvent(uint16 RouteId, EStyle Style, const FOnEventContext& Context)
diff --git a/Engine/Plugins/Experimental/AnimNext/Source/AnimNextEditor/Private/RewindDebugger/AnimNextModuleTrack.cpp b/Engine/Plugins/Experimental/AnimNext/Source/AnimNextEditor/Private/RewindDebugger/AnimNextModuleTrack.cpp
--- a/Engine/Plugins/Experimental/AnimNext/Source/AnimNextEditor/Private/RewindDebugger/AnimNextModuleTrack.cpp
+++ b/Engine/Plugins/Experimental/AnimNext/Source/AnimNextEditor/Private/RewindDebugger/AnimNextModuleTrack.cpp
@@ -275,7 +275,7 @@ bool FAnimNextModuleTrackCreator::HasDebugInfoInternal(const RewindDebugger::FOb
 		const FAnimNextProvider* AnimNextProvider = IRewindDebugger::Instance()->GetAnalysisSession()->ReadProvider<FAnimNextProvider>("AnimNextProvider");
     
 		uint64 ModuleId;
-		return AnimNextProvider->GetModuleId(ObjectId.GetChildElementId(), ModuleId);
+		return AnimNextProvider->GetModuleId(ObjectId.GetMainId(), ModuleId) != 0;
 	}
 
 	return false;
diff --git a/Engine/Plugins/Experimental/AnimNextAnimGraph/Source/AnimNextAnimGraph/Private/Graph/TraceAnimNextGraphInstances.cpp b/Engine/Plugins/Experimental/AnimNextAnimGraph/Source/AnimNextAnimGraph/Private/Graph/TraceAnimNextGraphInstances.cpp
--- a/Engine/Plugins/Experimental/AnimNextAnimGraph/Source/AnimNextAnimGraph/Private/Graph/TraceAnimNextGraphInstances.cpp
+++ b/Engine/Plugins/Experimental/AnimNextAnimGraph/Source/AnimNextAnimGraph/Private/Graph/TraceAnimNextGraphInstances.cpp
@@ -15,7 +15,7 @@
 
 #if ANIMNEXT_TRACE_ENABLED
 
-UE_TRACE_EVENT_BEGIN(AnimNextAnimGraph, EvaluationProgram)
+UE_TRACE_EVENT_BEGIN(UAFAnimGraph, EvaluationProgram)
 	UE_TRACE_EVENT_FIELD(uint64, Cycle)
 	UE_TRACE_EVENT_FIELD(double, RecordingTime)
 	UE_TRACE_EVENT_FIELD(uint64, OuterObjectId)
@@ -27,7 +27,7 @@ namespace UE::AnimNext
 {
 	void TraceGraphInstances(const FAnimNextGraphInstance& RootGraph)
 	{
-		bool bChannelEnabled = UE_TRACE_CHANNELEXPR_IS_ENABLED(AnimNextChannel);
+		bool bChannelEnabled = UE_TRACE_CHANNELEXPR_IS_ENABLED(UAFChannel);
     	if (!bChannelEnabled)
     	{
     		return;
@@ -72,7 +72,7 @@ namespace UE::AnimNext
 
 	void TraceEvaluationProgram(const FEvaluationProgram& Program, const FAnimNextGraphInstance& RootGraph)
 	{
-		bool bChannelEnabled = UE_TRACE_CHANNELEXPR_IS_ENABLED(AnimNextChannel);
+		bool bChannelEnabled = UE_TRACE_CHANNELEXPR_IS_ENABLED(UAFChannel);
 		if (!bChannelEnabled)
 		{
 			return;
@@ -110,7 +110,7 @@ namespace UE::AnimNext
 		FSerializableEvaluationProgram SerializableProgram(Program);
 		FSerializableEvaluationProgram::StaticStruct()->SerializeItem(Archive, &SerializableProgram, &Defaults);
 		
-		UE_TRACE_LOG(AnimNextAnimGraph, EvaluationProgram, AnimNextChannel)
+		UE_TRACE_LOG(UAFAnimGraph, EvaluationProgram, UAFChannel)
 				<< EvaluationProgram.Cycle(FPlatformTime::Cycles64())
 				<< EvaluationProgram.OuterObjectId(OuterObjectId)
 				<< EvaluationProgram.RecordingTime(RecordingTime)
diff --git a/Engine/Plugins/Experimental/AnimNextAnimGraph/Source/AnimNextAnimGraphEditor/Private/RewindDebugger/AnimNextAnimGraphAnalyzer.cpp b/Engine/Plugins/Experimental/AnimNextAnimGraph/Source/AnimNextAnimGraphEditor/Private/RewindDebugger/AnimNextAnimGraphAnalyzer.cpp
--- a/Engine/Plugins/Experimental/AnimNextAnimGraph/Source/AnimNextAnimGraphEditor/Private/RewindDebugger/AnimNextAnimGraphAnalyzer.cpp
+++ b/Engine/Plugins/Experimental/AnimNextAnimGraph/Source/AnimNextAnimGraphEditor/Private/RewindDebugger/AnimNextAnimGraphAnalyzer.cpp
@@ -20,7 +20,7 @@ void FAnimNextAnimGraphAnalyzer::OnAnalysisBegin(const FOnAnalysisContext& Conte
 {
 	auto& Builder = Context.InterfaceBuilder;
 
-	Builder.RouteEvent(RouteId_EvaluationProgram, "AnimNextAnimGraph", "EvaluationProgram");
+	Builder.RouteEvent(RouteId_EvaluationProgram, "UAFAnimGraph", "EvaluationProgram");
 }
 
 bool FAnimNextAnimGraphAnalyzer::OnEvent(uint16 RouteId, EStyle Style, const FOnEventContext& Context)
