UAF: Fix LOCTEXT conflicts

diff --git a/Engine/Plugins/Experimental/AnimNext/Source/AnimNext/Private/Component/AnimNextComponent.cpp b/Engine/Plugins/Experimental/AnimNext/Source/AnimNext/Private/Component/AnimNextComponent.cpp
--- a/Engine/Plugins/Experimental/AnimNext/Source/AnimNext/Private/Component/AnimNextComponent.cpp
+++ b/Engine/Plugins/Experimental/AnimNext/Source/AnimNext/Private/Component/AnimNextComponent.cpp
@@ -95,7 +95,7 @@ DEFINE_FUNCTION(UAnimNextComponent::execBlueprintSetVariable)
 	{
 		FBlueprintExceptionInfo ExceptionInfo(
 			EBlueprintExceptionType::NonFatalError,
-			NSLOCTEXT("UAFComponent", "UAFComponent_SetVariableInvalidWarning", "Invalid variable name supplied to Set Variable")
+			NSLOCTEXT("UAFComponent", "UAFComponent_SetVariableInvalidNameWarning", "Invalid variable name supplied to Set Variable")
 		);
 
 		FBlueprintCoreDelegates::ThrowScriptException(P_THIS, Stack, ExceptionInfo);
diff --git a/Engine/Plugins/Experimental/AnimNext/Source/AnimNextUncookedOnly/Private/UncookedOnlyUtils.cpp b/Engine/Plugins/Experimental/AnimNext/Source/AnimNextUncookedOnly/Private/UncookedOnlyUtils.cpp
--- a/Engine/Plugins/Experimental/AnimNext/Source/AnimNextUncookedOnly/Private/UncookedOnlyUtils.cpp
+++ b/Engine/Plugins/Experimental/AnimNext/Source/AnimNextUncookedOnly/Private/UncookedOnlyUtils.cpp
@@ -203,13 +203,13 @@ void FUtils::CompileVariables(const FRigVMCompileSettings& InSettings, UAnimNext
 					const UAnimNextSharedVariables* SharedVariablesAsset = SharedVariablesEntry->GetAsset();
 					if (SharedVariablesAsset == nullptr)
 					{
-						Log.Error(FText::Format(LOCTEXT("InvalidSharedVariablesFound", "Invalid shared variables asset found ({0})"), FText::FromString(SharedVariablesEntry->GetObjectPath().ToString())));
+						Log.Error(FText::Format(LOCTEXT("InvalidSharedVariableAssetFound", "Invalid shared variables asset found ({0})"), FText::FromString(SharedVariablesEntry->GetObjectPath().ToString())));
 						continue;
 					}
 
 					if (SharedVariableAssets.Contains(SharedVariablesAsset))
 					{
-						Log.Error(FText::Format(LOCTEXT("DuplicateSharedVariablesFound", "Duplicate shared variables asset found ({0})"), FText::FromString(SharedVariablesEntry->GetObjectPath().ToString())));
+						Log.Error(FText::Format(LOCTEXT("DuplicateSharedVariableAssetFound", "Duplicate shared variables asset found ({0})"), FText::FromString(SharedVariablesEntry->GetObjectPath().ToString())));
 						continue;
 					}
 
@@ -221,13 +221,13 @@ void FUtils::CompileVariables(const FRigVMCompileSettings& InSettings, UAnimNext
 					const UScriptStruct* SharedVariablesStruct = SharedVariablesEntry->GetStruct();
 					if (SharedVariablesStruct == nullptr)
 					{
-						Log.Error(FText::Format(LOCTEXT("InvalidSharedVariablesFound", "Invalid shared variables struct found ({0})"), FText::FromString(SharedVariablesEntry->GetObjectPath().ToString())));
+						Log.Error(FText::Format(LOCTEXT("InvalidSharedVariableStructFound", "Invalid shared variables struct found ({0})"), FText::FromString(SharedVariablesEntry->GetObjectPath().ToString())));
 						continue;
 					}
 
 					if (SharedVariableStructs.Contains(SharedVariablesStruct))
 					{
-						Log.Error(FText::Format(LOCTEXT("DuplicateSharedVariablesFound", "Duplicate shared variables struct found ({0})"), FText::FromString(SharedVariablesEntry->GetObjectPath().ToString())));
+						Log.Error(FText::Format(LOCTEXT("DuplicateSharedVariableStructFound", "Duplicate shared variables struct found ({0})"), FText::FromString(SharedVariablesEntry->GetObjectPath().ToString())));
 						continue;
 					}
 
