UAF: Fix overrides indexing

diff --git a/Engine/Plugins/Experimental/AnimNext/Source/AnimNext/Private/Variables/UAFInstanceVariableContainer.cpp b/Engine/Plugins/Experimental/AnimNext/Source/AnimNext/Private/Variables/UAFInstanceVariableContainer.cpp
--- a/Engine/Plugins/Experimental/AnimNext/Source/AnimNext/Private/Variables/UAFInstanceVariableContainer.cpp
+++ b/Engine/Plugins/Experimental/AnimNext/Source/AnimNext/Private/Variables/UAFInstanceVariableContainer.cpp
@@ -982,9 +982,9 @@ void FUAFInstanceVariableContainer::ResolveOverrides()
 			checkf(Asset == PinnedOverrides->Collection[OverridesIndex].AssetOrStructData.Get<FAssetType>(), TEXT("Variable overrides asset does not match"));
 			TConstArrayView<FVariableOverrides::FOverride> OverridesView = PinnedOverrides->Collection[OverridesIndex].Overrides;
 			const int32 NumOverrides = OverridesView.Num();
-			for (int32 VariableIndex = 0; VariableIndex < NumOverrides; ++VariableIndex)
+			for (int32 OverrideIndex = 0; OverrideIndex < NumOverrides; ++OverrideIndex)
 			{
-				const FVariableOverrides::FOverride& Override = OverridesView[VariableIndex];
+				const FVariableOverrides::FOverride& Override = OverridesView[OverrideIndex];
 				const FPropertyBagPropertyDesc* Desc = PropertyBag.FindPropertyDescByName(Override.Name);
 				if (Desc == nullptr)
 				{
@@ -996,6 +996,7 @@ void FUAFInstanceVariableContainer::ResolveOverrides()
 					ensureMsgf(false, TEXT("Variable type does not match"));
 					continue;
 				}
+				int32 VariableIndex = Desc - PropertyBag.GetPropertyBagStruct()->GetPropertyDescs().GetData();
 				ResolvedOverrides[VariableIndex] = Override.Memory;
 			}
 			break;
@@ -1003,24 +1004,32 @@ void FUAFInstanceVariableContainer::ResolveOverrides()
 	case FAssetOrStructType::IndexOfType<FStructType>():
 		{
 			const UScriptStruct* Struct = AssetOrStructData.Get<FStructType>();
+			check(StructDataCache.IsValid());
 			checkf(Struct == PinnedOverrides->Collection[OverridesIndex].AssetOrStructData.Get<FVariableOverrides::FStructType>(), TEXT("Variable overrides struct does not match"));
 			TConstArrayView<FVariableOverrides::FOverride> OverridesView = PinnedOverrides->Collection[OverridesIndex].Overrides;
 			const int32 NumOverrides = OverridesView.Num();
-			for (int32 VariableIndex = 0; VariableIndex < NumOverrides; ++VariableIndex)
+			for (int32 OverrideIndex = 0; OverrideIndex < NumOverrides; ++OverrideIndex)
 			{
-				const FVariableOverrides::FOverride& Override = OverridesView[VariableIndex];
-				const FProperty* Property = Struct->FindPropertyByName(Override.Name);
-				if (Property == nullptr)
+				const FVariableOverrides::FOverride& Override = OverridesView[OverrideIndex];
+				int32 PropertyIndex = StructDataCache->GetProperties().IndexOfByPredicate([&Override](const FStructDataCache::FPropertyInfo& InPropertyInfo)
+				{
+					return Override.Name == InPropertyInfo.Property->GetFName();
+				});
+				if (PropertyIndex == INDEX_NONE)
 				{
 					ensureMsgf(false, TEXT("Variable not found"));
-					continue;
+					continue;	
 				}
-				if (Override.Type != FAnimNextParamType::FromProperty(Property))
+
+				const FStructDataCache::FPropertyInfo& PropertyInfo = StructDataCache->GetProperties()[PropertyIndex];
+				check(PropertyInfo.Property != nullptr);
+
+				if (Override.Type != PropertyInfo.Type)
 				{
 					ensureMsgf(false, TEXT("Variable type does not match"));
 					continue;
 				}
-				ResolvedOverrides[VariableIndex] = Override.Memory;
+				ResolvedOverrides[PropertyIndex] = Override.Memory;
 			}
 			break;
 		}
