UAF: Fix check() when calling SetVariable on a component running a module that has errors

diff --git a/Engine/Plugins/Experimental/AnimNext/Source/AnimNext/Private/Variables/UAFInstanceVariableDataProxy.cpp b/Engine/Plugins/Experimental/AnimNext/Source/AnimNext/Private/Variables/UAFInstanceVariableDataProxy.cpp
--- a/Engine/Plugins/Experimental/AnimNext/Source/AnimNext/Private/Variables/UAFInstanceVariableDataProxy.cpp
+++ b/Engine/Plugins/Experimental/AnimNext/Source/AnimNext/Private/Variables/UAFInstanceVariableDataProxy.cpp
@@ -55,7 +55,11 @@ void FUAFInstanceVariableDataProxy::Reset()
 
 EPropertyBagResult FUAFInstanceVariableDataProxy::WriteVariable(const FAnimNextVariableReference& InVariable,  const FAnimNextParamType& InType, TArrayView<uint8>& OutResult)
 {
-	check(Variables);	// Not yet initialized
+	if(Variables == nullptr)
+	{
+		// Not yet initialized
+		return EPropertyBagResult::PropertyNotFound;
+	}
 	
 	const FUAFInstanceVariableData::FVariableMapping* FoundMapping = Variables->FindMapping(InVariable);
 	if (FoundMapping == nullptr)
@@ -74,7 +78,11 @@ EPropertyBagResult FUAFInstanceVariableDataProxy::WriteVariable(const FAnimNextV
 
 EPropertyBagResult FUAFInstanceVariableDataProxy::SetVariable(const FAnimNextVariableReference& InVariable, const FAnimNextParamType& InType, TConstArrayView<uint8> InNewValue)
 {
-	check(Variables);	// Not yet initialized
+	if(Variables == nullptr)
+	{
+		// Not yet initialized
+		return EPropertyBagResult::PropertyNotFound;
+	}
 
 	const FUAFInstanceVariableData::FVariableMapping* FoundMapping =  Variables->FindMapping(InVariable);
 	if (FoundMapping == nullptr)
@@ -93,7 +101,11 @@ EPropertyBagResult FUAFInstanceVariableDataProxy::SetVariable(const FAnimNextVar
 
 bool FUAFInstanceVariableDataProxy::WriteVariables(const UScriptStruct* InStruct, TFunctionRef<void(FStructView)> InFunction)
 {
-	check(Variables);	// Not yet initialized
+	if(Variables == nullptr)
+	{
+		// Not yet initialized
+		return false;
+	}
 
 	const int32* FoundIndexPtr = Variables->AllSharedVariableContainersMap.Find(InStruct);
 	if (FoundIndexPtr == nullptr)
