UAF module assets cannot be moved/copied

diff --git a/Engine/Plugins/Experimental/AnimNext/Source/AnimNextUncookedOnly/Internal/AnimNextRigVMAssetEditorData.h b/Engine/Plugins/Experimental/AnimNext/Source/AnimNextUncookedOnly/Internal/AnimNextRigVMAssetEditorData.h
--- a/Engine/Plugins/Experimental/AnimNext/Source/AnimNextUncookedOnly/Internal/AnimNextRigVMAssetEditorData.h
+++ b/Engine/Plugins/Experimental/AnimNext/Source/AnimNextUncookedOnly/Internal/AnimNextRigVMAssetEditorData.h
@@ -567,7 +567,7 @@ protected:
 	void UpgradeDataInterfaces();
 
 	/** All entries in this asset - not saved, discovered at load time and also contains InternalEntries */
-	UPROPERTY(transient)
+	UPROPERTY(Transient, DuplicateTransient)
 	TArray<TObjectPtr<UAnimNextRigVMAssetEntry>> Entries;
 
 	/** Internal entries in this asset */
diff --git a/Engine/Plugins/Experimental/AnimNext/Source/AnimNextUncookedOnly/Private/AnimNextRigVMAssetEditorData.cpp b/Engine/Plugins/Experimental/AnimNext/Source/AnimNextUncookedOnly/Private/AnimNextRigVMAssetEditorData.cpp
--- a/Engine/Plugins/Experimental/AnimNext/Source/AnimNextUncookedOnly/Private/AnimNextRigVMAssetEditorData.cpp
+++ b/Engine/Plugins/Experimental/AnimNext/Source/AnimNextUncookedOnly/Private/AnimNextRigVMAssetEditorData.cpp
@@ -112,12 +112,6 @@ void UAnimNextRigVMAssetEditorData::Serialize(FArchive& Ar)
 	RigVMClient.SetDefaultSchemaClass(UAnimNextRigVMAssetSchema::StaticClass());
 	RigVMClient.SetOuterClientHost(this, GET_MEMBER_NAME_CHECKED(UAnimNextRigVMAssetEditorData, RigVMClient));
 
-	const bool bIsDuplicating = (Ar.GetPortFlags() & PPF_Duplicate) != 0;
-	if (bIsDuplicating)
-	{
-		Ar << Entries;
-	}
-
 	Super::Serialize(Ar);
 
 	if (Ar.IsLoading() && Ar.CustomVer(FFortniteMainBranchObjectVersion::GUID) < FFortniteMainBranchObjectVersion::AnimNextVariableReferences)
