Support for displaying native structs in Anim Next rewind debugger anim graph tracks

diff --git a/Engine/Plugins/Experimental/AnimNext/Source/AnimNext/Private/RewindDebugger/AnimNextTrace.cpp b/Engine/Plugins/Experimental/AnimNext/Source/AnimNext/Private/RewindDebugger/AnimNextTrace.cpp
--- a/Engine/Plugins/Experimental/AnimNext/Source/AnimNext/Private/RewindDebugger/AnimNextTrace.cpp
+++ b/Engine/Plugins/Experimental/AnimNext/Source/AnimNext/Private/RewindDebugger/AnimNextTrace.cpp
@@ -29,6 +29,13 @@ UE_TRACE_EVENT_BEGIN(UAF, InstanceVariables)
 	UE_TRACE_EVENT_FIELD(uint8[], VariableData)
 UE_TRACE_EVENT_END()
 
+UE_TRACE_EVENT_BEGIN(UAF, InstanceVariablesStruct)
+	UE_TRACE_EVENT_FIELD(uint64, Cycle)
+	UE_TRACE_EVENT_FIELD(double, RecordingTime)
+	UE_TRACE_EVENT_FIELD(uint64, InstanceId)
+	UE_TRACE_EVENT_FIELD(uint8[], VariableData)
+UE_TRACE_EVENT_END()
+
 UE_TRACE_EVENT_BEGIN(UAF, InstanceVariableDescriptions)
 	UE_TRACE_EVENT_FIELD(uint32, VariableDescriptionHash)
 	UE_TRACE_EVENT_FIELD(uint8[], VariableDescriptionData)
@@ -151,15 +158,21 @@ void FAnimNextTrace::OutputAnimNextVariables(const FUAFAssetInstance* Instance,
 
 	OutputAnimNextInstance(Instance, OuterObject);
 
+	bool bVariablesOutput = false;
 	if (Instance->Variables.OwnedVariableContainers.Num() > 0)
 	{
 		for (const TSharedRef<FUAFInstanceVariableContainer>& OwnedVariableSet : Instance->Variables.OwnedVariableContainers)
 		{
-			OutputAnimNextVariableSet(OwnedVariableSet, InstanceId, OuterObject);
+			if (OutputAnimNextVariableSet(OwnedVariableSet, InstanceId, OuterObject))
+			{
+				bVariablesOutput = true;
+			}
 		}
 	}
-	else
+
+	if (!bVariablesOutput)
 	{
+		// trace an empty InstanceVariables message, as it's used for lifetime tracking currently
 		UE_TRACE_LOG(UAF, InstanceVariables, UAFChannel)
 				<< InstanceVariables.Cycle(FPlatformTime::Cycles64())
 				<< InstanceVariables.RecordingTime(FObjectTrace::GetWorldElapsedTime(OuterObject->GetWorld()))
@@ -167,8 +180,9 @@ void FAnimNextTrace::OutputAnimNextVariables(const FUAFAssetInstance* Instance,
 	}
 }
 
-void FAnimNextTrace::OutputAnimNextVariableSet(const TSharedRef<FUAFInstanceVariableContainer>& InVariableSet, uint64 InstanceId, const UObject* OuterObject)
+bool FAnimNextTrace::OutputAnimNextVariableSet(const TSharedRef<FUAFInstanceVariableContainer>& InVariableSet, uint64 InstanceId, const UObject* OuterObject)
 {
+	bool bOutputAnything = false;
 	switch (InVariableSet->VariablesContainer.GetIndex())
 	{
 	case FUAFInstanceVariableContainer::FVariablesContainerType::IndexOfType<FInstancedPropertyBag>():
@@ -224,6 +238,7 @@ void FAnimNextTrace::OutputAnimNextVariableSet(const TSharedRef<FUAFInstanceVari
 					PropertyBag->SerializeItem(Archive, InstancedPropertyBag.GetMutableValue().GetMemory(), nullptr);
 				}
 
+				bOutputAnything = true;
 				UE_TRACE_LOG(UAF, InstanceVariables, UAFChannel)
 						<< InstanceVariables.Cycle(FPlatformTime::Cycles64())
 						<< InstanceVariables.RecordingTime(FObjectTrace::GetWorldElapsedTime(OuterObject->GetWorld()))
@@ -235,13 +250,33 @@ void FAnimNextTrace::OutputAnimNextVariableSet(const TSharedRef<FUAFInstanceVari
 		}
 		case FUAFInstanceVariableContainer::FVariablesContainerType::IndexOfType<FInstancedStruct>():
 		{
-			// TODO: Trace native instanced structs
+			TArray<uint8> ArchiveData;
+			ArchiveData.Reserve(1024 * 10);
+
+			FMemoryWriter WriterArchive(ArchiveData);
+			FObjectAsTraceIdProxyArchive Archive(WriterArchive);
+
+			FInstancedStruct& InstancedStruct = InVariableSet->VariablesContainer.Get<FInstancedStruct>();
+			
+			if (InstancedStruct.IsValid())
+			{
+				FInstancedStruct::StaticStruct()->SerializeItem(Archive, &InstancedStruct, nullptr);
+			}
+
+			bOutputAnything = true;
+			UE_TRACE_LOG(UAF, InstanceVariablesStruct, UAFChannel)
+					<< InstanceVariablesStruct.Cycle(FPlatformTime::Cycles64())
+					<< InstanceVariablesStruct.RecordingTime(FObjectTrace::GetWorldElapsedTime(OuterObject->GetWorld()))
+					<< InstanceVariablesStruct.InstanceId(InstanceId)
+					<< InstanceVariablesStruct.VariableData(ArchiveData.GetData(), ArchiveData.Num());
 			break;
 		}
 	default:
 		checkNoEntry();
 		break;
 	}
+
+	return bOutputAnything;
 }
 
 }
diff --git a/Engine/Plugins/Experimental/AnimNext/Source/AnimNext/Public/RewindDebugger/AnimNextTrace.h b/Engine/Plugins/Experimental/AnimNext/Source/AnimNext/Public/RewindDebugger/AnimNextTrace.h
--- a/Engine/Plugins/Experimental/AnimNext/Source/AnimNext/Public/RewindDebugger/AnimNextTrace.h
+++ b/Engine/Plugins/Experimental/AnimNext/Source/AnimNext/Public/RewindDebugger/AnimNextTrace.h
@@ -38,7 +38,7 @@ struct FAnimNextTrace
 	ANIMNEXT_API static void OutputAnimNextVariables(const FUAFAssetInstance* DataInterface, const UObject* OuterObject);
 
 private:
-	ANIMNEXT_API static void OutputAnimNextVariableSet(const TSharedRef<FUAFInstanceVariableContainer>& VariableSet, uint64 InstanceId, const UObject* OuterObject);
+	ANIMNEXT_API static bool OutputAnimNextVariableSet(const TSharedRef<FUAFInstanceVariableContainer>& VariableSet, uint64 InstanceId, const UObject* OuterObject);
 };
 
 }
diff --git a/Engine/Plugins/Experimental/AnimNext/Source/AnimNextEditor/Private/RewindDebugger/AnimNextAnalyzer.cpp b/Engine/Plugins/Experimental/AnimNext/Source/AnimNextEditor/Private/RewindDebugger/AnimNextAnalyzer.cpp
--- a/Engine/Plugins/Experimental/AnimNext/Source/AnimNextEditor/Private/RewindDebugger/AnimNextAnalyzer.cpp
+++ b/Engine/Plugins/Experimental/AnimNext/Source/AnimNextEditor/Private/RewindDebugger/AnimNextAnalyzer.cpp
@@ -24,6 +24,7 @@ void FAnimNextAnalyzer::OnAnalysisBegin(const FOnAnalysisContext& Context)
 
 	Builder.RouteEvent(RouteId_Module, "UAF", "Instance");
 	Builder.RouteEvent(RouteId_InstanceVariables, "UAF", "InstanceVariables");
+	Builder.RouteEvent(RouteId_InstanceVariablesStruct, "UAF", "InstanceVariablesStruct");
 	Builder.RouteEvent(RouteId_InstanceVariableDescriptions, "UAF", "InstanceVariableDescriptions");
 }
 
@@ -53,7 +54,17 @@ bool FAnimNextAnalyzer::OnEvent(uint16 RouteId, EStyle Style, const FOnEventCont
 			uint32 VariableDescHash = EventData.GetValue<uint32>("VariableDescriptionHash");
 			TArrayView<const uint8> VariableData = EventData.GetArrayView<uint8>("VariableData");
 
-			Provider.AppendVariables(Context.EventTime.AsSeconds(Cycle), RecordingTime, ModuleInstanceId, VariableDescHash, VariableData);
+			Provider.AppendVariables(Context.EventTime.AsSeconds(Cycle), RecordingTime, ModuleInstanceId, EPropertyVariableDataType::PropertyBag, VariableDescHash, VariableData);
+			break;
+		}
+		case RouteId_InstanceVariablesStruct:
+		{
+			uint64 ModuleInstanceId = EventData.GetValue<uint64>("InstanceId");
+			uint64 Cycle = EventData.GetValue<uint64>("Cycle");
+			double RecordingTime = EventData.GetValue<double>("RecordingTime");
+			TArrayView<const uint8> VariableData = EventData.GetArrayView<uint8>("VariableData");
+
+			Provider.AppendVariables(Context.EventTime.AsSeconds(Cycle), RecordingTime, ModuleInstanceId, EPropertyVariableDataType::InstancedStruct, 0, VariableData);
 			break;
 		}
 		case RouteId_InstanceVariableDescriptions:
diff --git a/Engine/Plugins/Experimental/AnimNext/Source/AnimNextEditor/Private/RewindDebugger/AnimNextAnalyzer.h b/Engine/Plugins/Experimental/AnimNext/Source/AnimNextEditor/Private/RewindDebugger/AnimNextAnalyzer.h
--- a/Engine/Plugins/Experimental/AnimNext/Source/AnimNextEditor/Private/RewindDebugger/AnimNextAnalyzer.h
+++ b/Engine/Plugins/Experimental/AnimNext/Source/AnimNextEditor/Private/RewindDebugger/AnimNextAnalyzer.h
@@ -24,6 +24,7 @@ private:
 	{
 		RouteId_Module,
 		RouteId_InstanceVariables,
+		RouteId_InstanceVariablesStruct,
 		RouteId_InstanceVariableDescriptions
 	};
 
diff --git a/Engine/Plugins/Experimental/AnimNext/Source/AnimNextEditor/Private/RewindDebugger/AnimNextModuleTrack.cpp b/Engine/Plugins/Experimental/AnimNext/Source/AnimNextEditor/Private/RewindDebugger/AnimNextModuleTrack.cpp
--- a/Engine/Plugins/Experimental/AnimNext/Source/AnimNextEditor/Private/RewindDebugger/AnimNextModuleTrack.cpp
+++ b/Engine/Plugins/Experimental/AnimNext/Source/AnimNextEditor/Private/RewindDebugger/AnimNextModuleTrack.cpp
@@ -153,43 +153,55 @@ bool FAnimNextModuleTrack::UpdateInternal()
 					const TraceServices::IFrameProvider& FramesProvider = TraceServices::ReadFrameProvider(*AnalysisSession);
 					TraceServices::FAnalysisSessionReadScope SessionReadScope(*AnalysisSession);
 
+					
+					DetailsObject->Properties.SetNum(0, EAllowShrinking::No);
+					DetailsObject->NativeProperties.SetNum(0, EAllowShrinking::No);
+
 					TraceServices::FFrame MarkerFrame;
 					if(FramesProvider.GetFrameFromTime(ETraceFrameType::TraceFrameType_Game, CurrentScrubTime, MarkerFrame))
 					{
 						Data->VariablesTimeline.EnumerateEvents(MarkerFrame.StartTime, MarkerFrame.EndTime, [AnimNextProvider, DetailsObject, AnalysisSession, this](double InStartTime, double InEndTime, uint32 InDepth, const FPropertyVariableData& VariableListData)
 						{
-							// First look up property description
-							const FPropertyDescriptionData* DescriptionData = AnimNextProvider->GetPropertyDescriptionData(VariableListData.DescriptionHash);
-							if (DescriptionData == nullptr)
-							{
-								// Can't do anything without a description
-								return TraceServices::EEventEnumerate::Continue;
-							}
-							
 							const IGameplayProvider* GameplayProvider = AnalysisSession->ReadProvider<IGameplayProvider>("GameplayProvider");
-							
+							if (VariableListData.ValueType == EPropertyVariableDataType::PropertyBag)
 							{
+								// First look up property description
+								const FPropertyDescriptionData* DescriptionData = AnimNextProvider->GetPropertyDescriptionData(VariableListData.DescriptionHash);
+								if (DescriptionData == nullptr)
+								{
+									// Can't do anything without a description
+									return TraceServices::EEventEnumerate::Continue;
+								}
+								
 								// Load the property descriptions
-								FMemoryReader Reader(DescriptionData->Data);
-								FObjectAsTraceIdProxyArchiveReader ReaderProxy(Reader, GameplayProvider);
-								ReaderProxy.UsingCustomVersion(AnimNext::FAnimNextTrace::CustomVersionGUID);
-								ReaderProxy << PropertyDescriptions;
-
-								DetailsObject->Properties.Reset();
-								DetailsObject->Properties.AddProperties(PropertyDescriptions);
-							}
-
-							{
+								FMemoryReader DescriptionReader(DescriptionData->Data);
+								FObjectAsTraceIdProxyArchiveReader DescriptionReaderProxy(DescriptionReader, GameplayProvider);
+								DescriptionReaderProxy.UsingCustomVersion(AnimNext::FAnimNextTrace::CustomVersionGUID);
+								DescriptionReaderProxy << PropertyDescriptions;
+								DetailsObject->Properties.SetNum(DetailsObject->Properties.Num()+1);
+								FInstancedPropertyBag &PropertyBag = DetailsObject->Properties.Last();
+								PropertyBag.AddProperties(PropertyDescriptions);
+							
 								// Load the property values
 								FMemoryReader Reader(VariableListData.ValueData);
 								FObjectAsTraceIdProxyArchiveReader ReaderProxy(Reader, GameplayProvider);
 
-								UPropertyBag* PropertyBag = const_cast<UPropertyBag*>(DetailsObject->Properties.GetPropertyBagStruct());
-								if (PropertyBag != nullptr)
+								UPropertyBag* PropertyBagStruct = const_cast<UPropertyBag*>(DetailsObject->Properties.Last().GetPropertyBagStruct());
+								if (PropertyBagStruct != nullptr)
 								{
-									PropertyBag->SerializeItem(ReaderProxy, DetailsObject->Properties.GetMutableValue().GetMemory(), nullptr);
+									PropertyBagStruct->SerializeItem(ReaderProxy, PropertyBag.GetMutableValue().GetMemory(), nullptr);
 								}
 							}
+							else //(VariableListData.ValueType == EPropertyVariableDataType::InstancedStruct)
+							{
+								FMemoryReader Reader(VariableListData.ValueData);
+								FObjectAsTraceIdProxyArchiveReader ReaderProxy(Reader, GameplayProvider);
+								
+								DetailsObject->NativeProperties.SetNum(DetailsObject->NativeProperties.Num()+1);
+								FInstancedStruct& InstancedStruct = DetailsObject->NativeProperties.Last();
+								
+								FInstancedStruct::StaticStruct()->SerializeItem(ReaderProxy, &InstancedStruct, nullptr);
+							}
 
 							return TraceServices::EEventEnumerate::Stop;
 						});
diff --git a/Engine/Plugins/Experimental/AnimNext/Source/AnimNextEditor/Private/RewindDebugger/AnimNextModuleTrack.h b/Engine/Plugins/Experimental/AnimNext/Source/AnimNextEditor/Private/RewindDebugger/AnimNextModuleTrack.h
--- a/Engine/Plugins/Experimental/AnimNext/Source/AnimNextEditor/Private/RewindDebugger/AnimNextModuleTrack.h
+++ b/Engine/Plugins/Experimental/AnimNext/Source/AnimNextEditor/Private/RewindDebugger/AnimNextModuleTrack.h
@@ -20,8 +20,11 @@ class UPropertyBagDetailsObject : public UObject
 	GENERATED_BODY()
 
 public:
-	UPROPERTY(EditAnywhere, Category=Details, meta=(ShowOnlyInnerProperties))
-	FInstancedPropertyBag Properties;
+	UPROPERTY(EditAnywhere, Category=Details)
+	TArray<FInstancedPropertyBag> Properties;
+
+	UPROPERTY(EditAnywhere, Category=Details)
+	TArray<FInstancedStruct> NativeProperties;
 };
 
 #if ANIMNEXT_TRACE_ENABLED
diff --git a/Engine/Plugins/Experimental/AnimNext/Source/AnimNextEditor/Private/RewindDebugger/AnimNextProvider.cpp b/Engine/Plugins/Experimental/AnimNext/Source/AnimNextEditor/Private/RewindDebugger/AnimNextProvider.cpp
--- a/Engine/Plugins/Experimental/AnimNext/Source/AnimNextEditor/Private/RewindDebugger/AnimNextProvider.cpp
+++ b/Engine/Plugins/Experimental/AnimNext/Source/AnimNextEditor/Private/RewindDebugger/AnimNextProvider.cpp
@@ -34,7 +34,7 @@ void FAnimNextProvider::AppendInstance(uint64 InstanceId, uint64 HostInstanceId,
 	ChildList.AddUnique(Data);
 }
 
-void FAnimNextProvider::AppendVariables(double ProfileTime, double RecordingTime, uint64 DataInterfaceId, uint32 PropertyDescriptionHash, const TArrayView<const uint8>& VariableData)
+void FAnimNextProvider::AppendVariables(double ProfileTime, double RecordingTime, uint64 DataInterfaceId, EPropertyVariableDataType PropertyVariableDataType, uint32 PropertyDescriptionHash, const TArrayView<const uint8>& VariableData)
 {
 	Session.WriteAccessCheck();
 	
@@ -42,7 +42,7 @@ void FAnimNextProvider::AppendVariables(double ProfileTime, double RecordingTime
 	{
 		if (!VariableData.IsEmpty())
 		{
-			(*Data)->VariablesTimeline.AppendEvent(ProfileTime, FPropertyVariableData(PropertyDescriptionHash, VariableData));
+			(*Data)->VariablesTimeline.AppendEvent(ProfileTime, FPropertyVariableData(PropertyVariableDataType, PropertyDescriptionHash, VariableData));
 		}
 		(*Data)->EndTime = RecordingTime;
 		if ((*Data)->StartTime < 0)
diff --git a/Engine/Plugins/Experimental/AnimNext/Source/AnimNextEditor/Private/RewindDebugger/AnimNextProvider.h b/Engine/Plugins/Experimental/AnimNext/Source/AnimNextEditor/Private/RewindDebugger/AnimNextProvider.h
--- a/Engine/Plugins/Experimental/AnimNext/Source/AnimNextEditor/Private/RewindDebugger/AnimNextProvider.h
+++ b/Engine/Plugins/Experimental/AnimNext/Source/AnimNextEditor/Private/RewindDebugger/AnimNextProvider.h
@@ -20,11 +20,18 @@ struct FPropertyDescriptionData
 	TArray<uint8> Data;
 };
 
+enum class EPropertyVariableDataType
+{
+	PropertyBag,
+	InstancedStruct,
+};
+
 struct FPropertyVariableData
 {
-	FPropertyVariableData(uint32 DescriptionHash, TConstArrayView<uint8> Data)
-		: DescriptionHash(DescriptionHash) , ValueData(Data) {}
-	
+	FPropertyVariableData(EPropertyVariableDataType VariableDataType, uint32 DescriptionHash, TConstArrayView<uint8> Data)
+		: ValueType (VariableDataType), DescriptionHash(DescriptionHash) , ValueData(Data) {}
+
+	EPropertyVariableDataType ValueType;
 	uint32 DescriptionHash;
 	TArray<uint8> ValueData;
 };
@@ -59,7 +66,7 @@ public:
 	FAnimNextProvider(TraceServices::IAnalysisSession& InSession);
 
 	void AppendInstance(uint64 InstanceId, uint64 HostInstanceId,  uint64 AssetId, uint64 OuterObjectId);
-	void AppendVariables(double ProfileTime, double RecordingTime, uint64 InModuleInstanceId, uint32 PropertyDescriptionHash, const TArrayView<const uint8>& VariableData);
+	void AppendVariables(double ProfileTime, double RecordingTime, uint64 InModuleInstanceId, EPropertyVariableDataType PropertyVariableDataType, uint32 PropertyDescriptionHash, const TArrayView<const uint8>& VariableData);
 	void AppendVariableDescriptions(uint32 PropertyDescriptionHash, const TArrayView<const uint8>& VariableDescriptionData);
 
 	bool GetModuleId(uint64 ComponentId, uint64& OutModuleId) const;
