Don't create internal variables for execution IO, visible only, or hidden pins.
Fix ST RigVM task / condition result indexes getting cleared on compile.

diff --git a/Engine/Plugins/Experimental/AnimNext/Source/AnimNextUncookedOnly/Private/AnimNextRigVMAssetEditorData.cpp b/Engine/Plugins/Experimental/AnimNext/Source/AnimNextUncookedOnly/Private/AnimNextRigVMAssetEditorData.cpp
--- a/Engine/Plugins/Experimental/AnimNext/Source/AnimNextUncookedOnly/Private/AnimNextRigVMAssetEditorData.cpp
+++ b/Engine/Plugins/Experimental/AnimNext/Source/AnimNextUncookedOnly/Private/AnimNextRigVMAssetEditorData.cpp
@@ -808,9 +808,13 @@ void UAnimNextRigVMAssetEditorData::BuildFunctionWrapperEventVariables(FAnimNext
 		const FRigVMGraphFunctionHeader& FunctionHeader = ProgrammaticFunctionHeader.FunctionHeader;
 		for (const FRigVMGraphFunctionArgument& Argument : FunctionHeader.Arguments)
 		{
-			FRigVMGraphFunctionArgument InternallyNamedArgument = Argument;
-			InternallyNamedArgument.Name = FName(FUtils::MakeFunctionWrapperVariableName(FunctionHeader.Name, Argument.Name));
-			InContext.ProgrammaticVariables.Add(FAnimNextProgrammaticVariable::FromRigVMGraphFunctionArgument(InternallyNamedArgument));
+			// Don't create internal variables for execution IO, visible only, or hidden pins. Those types do not need variable nodes in the graph.
+			if (Argument.Direction == ERigVMPinDirection::Input || Argument.Direction == ERigVMPinDirection::Output)
+			{
+				FRigVMGraphFunctionArgument InternallyNamedArgument = Argument;
+				InternallyNamedArgument.Name = FName(FUtils::MakeFunctionWrapperVariableName(FunctionHeader.Name, Argument.Name));
+				InContext.ProgrammaticVariables.Add(FAnimNextProgrammaticVariable::FromRigVMGraphFunctionArgument(InternallyNamedArgument));
+			}
 		}
 	}
 }
diff --git a/Engine/Plugins/Experimental/AnimNextStateTree/Source/AnimNextStateTree/Private/Conditions/AnimNextStateTreeRigVMConditionBase.cpp b/Engine/Plugins/Experimental/AnimNextStateTree/Source/AnimNextStateTree/Private/Conditions/AnimNextStateTreeRigVMConditionBase.cpp
--- a/Engine/Plugins/Experimental/AnimNextStateTree/Source/AnimNextStateTree/Private/Conditions/AnimNextStateTreeRigVMConditionBase.cpp
+++ b/Engine/Plugins/Experimental/AnimNextStateTree/Source/AnimNextStateTree/Private/Conditions/AnimNextStateTreeRigVMConditionBase.cpp
@@ -150,7 +150,6 @@ void FAnimNextStateTreeRigVMConditionBase::GetProgrammaticFunctionHeaders(FAnimN
 	
 	StateName = State->Name;
 	NodeId = Desc.ID;
-	ResultIndex = INDEX_NONE;
 	InternalEventName = FName(FUtils::MakeFunctionWrapperEventName(RigVMFunctionHeader.Name));
 
 	FAnimNextProgrammaticFunctionHeader AnimNextFunctionHeader = {};
diff --git a/Engine/Plugins/Experimental/AnimNextStateTree/Source/AnimNextStateTree/Private/Tasks/AnimNextStateTreeRigVMTaskBase.cpp b/Engine/Plugins/Experimental/AnimNextStateTree/Source/AnimNextStateTree/Private/Tasks/AnimNextStateTreeRigVMTaskBase.cpp
--- a/Engine/Plugins/Experimental/AnimNextStateTree/Source/AnimNextStateTree/Private/Tasks/AnimNextStateTreeRigVMTaskBase.cpp
+++ b/Engine/Plugins/Experimental/AnimNextStateTree/Source/AnimNextStateTree/Private/Tasks/AnimNextStateTreeRigVMTaskBase.cpp
@@ -152,7 +152,6 @@ void FAnimNextStateTreeRigVMTaskBase::GetProgrammaticFunctionHeaders(FAnimNextSt
 	
 	StateName = State->Name;
 	NodeId = Desc.ID;
-	ResultIndex = INDEX_NONE;
 	InternalEventName = FName(FUtils::MakeFunctionWrapperEventName(RigVMFunctionHeader.Name));
 
 	FAnimNextProgrammaticFunctionHeader AnimNextFunctionHeader = {};
