AnimNext variables UI/UX changes:

* Added export caching behaviour to UAnimNextAnimGraphWorkspaceAssetUserData
* Added AnimNext Variable Find/Replace class, UI and behavior
* Added Promote to Variable schema action + behaviour
* Added ::GetDefaultValueString ::GetVariableCategory and ::SetVariableCategory to IAnimNextRigVMVariableInterface
* Added EAnimNextExportedVariableFlags filter to Variable Picker arguments
* Added Find/Replace context menu entries for (Shared)Variable graph nodes
* Added inclusion of export asset reference chain their variables when generating Graph actions
* Added drag-drop operations to Variables View
* Added additional command actions to Variables View
* Added the ability to add (anim next) variables to categories
* Added the ability to re-order variables within AnimNext assets
* Added representation of UStruct shared variable sources to Variables View
* Added/changed (script exposed) API for adding/refreshing Shared Variable nodes
* Added APIs for refactoring variables (deletion, renaming, moving, changing type), which automatically updates references across the project/specific assets

* Changed UI layout of Variables View
* Changed Variables View hierarchy generation, now generated top-down list of export its asset chain
* Changed Variables View now to be populated from a FWorkspaceOutlinerItemExport rather than an UObject
* Changed AR AnimNext/Workspace export definition/APIs

* Renamed DataInterface -> SharedVariables across various object/struct names

* Removed UE::AnimNext::Tests::FDataInterfaceCompile references (no longer exists)
* Removed the ability to switch public/private export on SharedVariable asset their variables within the Variables View

Workspace Editor changes:

* Added GetFocussedWorkspaceDocument, returns previously activated document tab its document (empty document if none set)
* Added GetAssetsOfType to FWorkspaceOutlinerItemExport
* Added tooltip to Workspace Outliner items
* Added concept of item groups to Workspace Outliner, driven by AR exports
* Added ability for IWorkspaceEditor::GetAssets to return set of FAssetData including any expanded asset reference items
* Added SceneOutliner item selection to UWorkspaceItemMenuContext

* Changed double-click behaviour, now also expands the clicked item
* Changed OnWorkspaceFocussedDocumentChanged, FindObjectDocumentType, FOnRedirectWorkspaceContext to use FWorkspaceDocument instead of UObject

diff --git a/Engine/Plugins/Experimental/AnimNext/Source/AnimNext/Public/Variables/AnimNextSharedVariables.h b/Engine/Plugins/Experimental/AnimNext/Source/AnimNext/Public/Variables/AnimNextSharedVariables.h
--- a/Engine/Plugins/Experimental/AnimNext/Source/AnimNext/Public/Variables/AnimNextSharedVariables.h
+++ b/Engine/Plugins/Experimental/AnimNext/Source/AnimNext/Public/Variables/AnimNextSharedVariables.h
@@ -14,11 +14,6 @@ namespace UE::AnimNext
 struct FInjectionInfo;
 }
 
-namespace UE::AnimNext::Tests
-{
-class FDataInterfaceCompile;
-}
-
 namespace UE::AnimNext::UncookedOnly
 {
 struct FUtils;
@@ -35,7 +30,6 @@ public:
 
 private:
 	friend class UAnimNextSharedVariablesFactory;
-	friend class UE::AnimNext::Tests::FDataInterfaceCompile;
 	friend struct UE::AnimNext::UncookedOnly::FUtils;
 	friend UE::AnimNext::FInjectionInfo;
 };
diff --git a/Engine/Plugins/Experimental/AnimNext/Source/AnimNextEditor/AnimNextEditor.Build.cs b/Engine/Plugins/Experimental/AnimNext/Source/AnimNextEditor/AnimNextEditor.Build.cs
--- a/Engine/Plugins/Experimental/AnimNext/Source/AnimNextEditor/AnimNextEditor.Build.cs
+++ b/Engine/Plugins/Experimental/AnimNext/Source/AnimNextEditor/AnimNextEditor.Build.cs
@@ -55,6 +55,8 @@ namespace UnrealBuildTool.Rules
 					"TraceInsights",
 					"RewindDebuggerInterface",
 					"GameplayInsights",
+					"Projects",
+					"Persona"
 				}
 			);
 		}
diff --git a/Engine/Plugins/Experimental/AnimNext/Source/AnimNextEditor/Internal/Common/AnimNextAssetItemDetails.cpp b/Engine/Plugins/Experimental/AnimNext/Source/AnimNextEditor/Internal/Common/AnimNextAssetItemDetails.cpp
--- a/Engine/Plugins/Experimental/AnimNext/Source/AnimNextEditor/Internal/Common/AnimNextAssetItemDetails.cpp
+++ b/Engine/Plugins/Experimental/AnimNext/Source/AnimNextEditor/Internal/Common/AnimNextAssetItemDetails.cpp
@@ -7,6 +7,7 @@
 #include "Entries/AnimNextEventGraphEntry.h"
 #include "InstancedStruct.h"
 #include "AnimNextAssetWorkspaceAssetUserData.h"
+#include "ISceneOutlinerTreeItem.h"
 #include "RigVMModel/RigVMGraph.h"
 #include "WorkspaceItemMenuContext.h"
 #include "IWorkspaceEditor.h"
@@ -26,6 +27,7 @@
 #include "UObject/UObjectIterator.h"
 #include "Variables/SAddVariablesDialog.h"
 #include "Widgets/Input/STextEntryPopup.h"
+#include "Variables/SVariablesView.h"
 
 #define LOCTEXT_NAMESPACE "FAnimNextGraphItemDetails"
 
@@ -54,6 +56,43 @@ const FSlateBrush* FAnimNextAssetItemDetails::GetItemIcon(const FWorkspaceOutlin
 	return nullptr;
 }
 
+bool FAnimNextAssetItemDetails::HandleDoubleClick(const FToolMenuContext& ToolMenuContext) const
+{
+	UWorkspaceItemMenuContext* WorkspaceItemContext = ToolMenuContext.FindContext<UWorkspaceItemMenuContext>();
+	const UAssetEditorToolkitMenuContext* AssetEditorContext = ToolMenuContext.FindContext<UAssetEditorToolkitMenuContext>();
+	if (WorkspaceItemContext == nullptr || AssetEditorContext == nullptr)
+	{
+		return false;
+	}
+
+	const TSharedPtr<UE::Workspace::IWorkspaceEditor> WorkspaceEditor = StaticCastSharedPtr<UE::Workspace::IWorkspaceEditor>(AssetEditorContext->Toolkit.Pin());
+	if(!WorkspaceEditor.IsValid())
+	{
+		return false;
+	}
+
+	// Double-clicking a SharedVariables entry will open it directly in the variables view
+	if (WorkspaceItemContext->SelectedExports.Num() == 1)
+	{
+		const FWorkspaceOutlinerItemExport& SharedVariablesExport = WorkspaceItemContext->SelectedExports[0];
+		const TInstancedStruct<FWorkspaceOutlinerItemData>& Data = SharedVariablesExport.GetResolvedExport().GetData();
+		if (Data.IsValid() && Data.GetScriptStruct()->IsChildOf(FAnimNextSharedVariablesOutlinerData::StaticStruct()))
+		{
+			const FTabId VariablesTabId(VariablesTabName);
+			TSharedPtr<SDockTab> VariableDockTab = WorkspaceEditor->GetTabManager()->FindExistingLiveTab(VariablesTabId);
+			if (VariableDockTab.IsValid())
+			{
+				const TSharedRef<SVariablesView> VariablesView = StaticCastSharedRef<SVariablesView>(VariableDockTab->GetContent());
+				VariablesView->SetExportDirectly(SharedVariablesExport);
+
+				return true;
+			}
+		}
+	}
+
+	return false;
+}
+
 void FAnimNextAssetItemDetails::RegisterToolMenuExtensions()
 {
 	FToolMenuOwnerScoped OwnerScoped(TEXT("FAnimNextModuleItemDetails"));
diff --git a/Engine/Plugins/Experimental/AnimNext/Source/AnimNextEditor/Internal/Common/AnimNextAssetItemDetails.h b/Engine/Plugins/Experimental/AnimNext/Source/AnimNextEditor/Internal/Common/AnimNextAssetItemDetails.h
--- a/Engine/Plugins/Experimental/AnimNext/Source/AnimNextEditor/Internal/Common/AnimNextAssetItemDetails.h
+++ b/Engine/Plugins/Experimental/AnimNext/Source/AnimNextEditor/Internal/Common/AnimNextAssetItemDetails.h
@@ -16,6 +16,7 @@ public:
 	FAnimNextAssetItemDetails() = default;
 
 	ANIMNEXTEDITOR_API virtual const FSlateBrush* GetItemIcon(const FWorkspaceOutlinerItemExport& Export) const override;
+	ANIMNEXTEDITOR_API virtual bool HandleDoubleClick(const FToolMenuContext& ToolMenuContext) const override;
 
 	static void RegisterToolMenuExtensions();
 	static void UnregisterToolMenuExtensions();
diff --git a/Engine/Plugins/Experimental/AnimNext/Source/AnimNextEditor/Internal/Common/GraphEditorSchemaActions.h b/Engine/Plugins/Experimental/AnimNext/Source/AnimNextEditor/Internal/Common/GraphEditorSchemaActions.h
--- a/Engine/Plugins/Experimental/AnimNext/Source/AnimNextEditor/Internal/Common/GraphEditorSchemaActions.h
+++ b/Engine/Plugins/Experimental/AnimNext/Source/AnimNextEditor/Internal/Common/GraphEditorSchemaActions.h
@@ -6,7 +6,6 @@
 #include "EdGraph/EdGraphSchema.h"
 #include "RigVMCore/RigVMExternalVariable.h"
 #include "RigVMCore/RigVMGraphFunctionDefinition.h"
-#include "Styling/AppStyle.h"
 #include "GraphEditorSchemaActions.generated.h"
 
 #define UE_API ANIMNEXTEDITOR_API
@@ -48,7 +47,7 @@ struct FAnimNextSchemaAction_RigUnit : public FAnimNextSchemaAction
 
 	// FEdGraphSchemaAction Interface
 	UE_API virtual UEdGraphNode* PerformAction(UEdGraph* ParentGraph, TArray<UEdGraphPin*>& FromPins, const FVector2f& Location, bool bSelectNewNode = true) override;
-	virtual UEdGraphNode* PerformAction(UEdGraph* ParentGraph, UEdGraphPin* FromPin, const FVector2f& Location, bool bSelectNewNode = true) { return nullptr; }
+	virtual UEdGraphNode* PerformAction(UEdGraph* ParentGraph, UEdGraphPin* FromPin, const FVector2f& Location, bool bSelectNewNode = true) override { return nullptr; }
 
 private:
 	// The script struct for our rig unit
@@ -73,7 +72,7 @@ struct FAnimNextSchemaAction_DispatchFactory : public FAnimNextSchemaAction
 	
 	// FEdGraphSchemaAction Interface
 	virtual UEdGraphNode* PerformAction(UEdGraph* ParentGraph, TArray<UEdGraphPin*>& FromPins, const FVector2f& Location, bool bSelectNewNode = true) override;
-	virtual UEdGraphNode* PerformAction(UEdGraph* ParentGraph, UEdGraphPin* FromPin, const FVector2f& Location, bool bSelectNewNode = true) { return nullptr; }
+	virtual UEdGraphNode* PerformAction(UEdGraph* ParentGraph, UEdGraphPin* FromPin, const FVector2f& Location, bool bSelectNewNode = true) override { return nullptr; }
 
 private:
 	// Notation for dispatch factory
@@ -99,7 +98,7 @@ struct FAnimNextSchemaAction_Variable : public FAnimNextSchemaAction
 
 	// FEdGraphSchemaAction Interface
 	UE_API virtual UEdGraphNode* PerformAction(UEdGraph* ParentGraph, TArray<UEdGraphPin*>& FromPins, const FVector2f& Location, bool bSelectNewNode = true) override;
-	virtual UEdGraphNode* PerformAction(UEdGraph* ParentGraph, UEdGraphPin* FromPin, const FVector2f& Location, bool bSelectNewNode = true) { return nullptr; }
+	virtual UEdGraphNode* PerformAction(UEdGraph* ParentGraph, UEdGraphPin* FromPin, const FVector2f& Location, bool bSelectNewNode = true) override { return nullptr; }
 
 	UE_API virtual const FSlateBrush* GetIconBrush() const override;
 
@@ -123,7 +122,7 @@ struct FAnimNextSchemaAction_AddComment : public FAnimNextSchemaAction
 
 	// FEdGraphSchemaAction Interface
 	virtual UEdGraphNode* PerformAction(UEdGraph* ParentGraph, TArray<UEdGraphPin*>& FromPins, const FVector2f& Location, bool bSelectNewNode = true) override;
-	virtual UEdGraphNode* PerformAction(UEdGraph* ParentGraph, UEdGraphPin* FromPin, const FVector2f& Location, bool bSelectNewNode = true) { return nullptr; }
+	virtual UEdGraphNode* PerformAction(UEdGraph* ParentGraph, UEdGraphPin* FromPin, const FVector2f& Location, bool bSelectNewNode = true) override { return nullptr; }
 
 	virtual const FSlateBrush* GetIconBrush() const override;
 };
@@ -142,7 +141,7 @@ struct FAnimNextSchemaAction_Function : public FAnimNextSchemaAction
 	
 	// FEdGraphSchemaAction Interface
 	virtual UEdGraphNode* PerformAction(UEdGraph* ParentGraph, TArray<UEdGraphPin*>& FromPins, const FVector2f& Location, bool bSelectNewNode = true) override;
-	virtual UEdGraphNode* PerformAction(UEdGraph* ParentGraph, UEdGraphPin* FromPin, const FVector2f& Location, bool bSelectNewNode = true) { return nullptr; }
+	virtual UEdGraphNode* PerformAction(UEdGraph* ParentGraph, UEdGraphPin* FromPin, const FVector2f& Location, bool bSelectNewNode = true) override { return nullptr; }
 
 private:
 	/** The public function definition we will spawn from [optional] */
@@ -158,4 +157,18 @@ private:
 	TSubclassOf<UEdGraphNode> NodeClass;
 };
 
+USTRUCT()
+struct FAnimNextSchemaAction_PromoteToVariable : public FAnimNextSchemaAction
+{
+	GENERATED_BODY()
+
+	FAnimNextSchemaAction_PromoteToVariable();
+
+	// FEdGraphSchemaAction Interface
+	virtual UEdGraphNode* PerformAction(UEdGraph* ParentGraph, TArray<UEdGraphPin*>& FromPins, const FVector2f& Location, bool bSelectNewNode = true) override;
+	virtual UEdGraphNode* PerformAction(UEdGraph* ParentGraph, UEdGraphPin* FromPin, const FVector2f& Location, bool bSelectNewNode = true) override { return nullptr; }
+
+	virtual const FSlateBrush* GetIconBrush() const override;
+};
+
 #undef UE_API
diff --git a/Engine/Plugins/Experimental/AnimNext/Source/AnimNextEditor/Internal/Common/SActionMenu.h b/Engine/Plugins/Experimental/AnimNext/Source/AnimNextEditor/Internal/Common/SActionMenu.h
--- a/Engine/Plugins/Experimental/AnimNext/Source/AnimNextEditor/Internal/Common/SActionMenu.h
+++ b/Engine/Plugins/Experimental/AnimNext/Source/AnimNextEditor/Internal/Common/SActionMenu.h
@@ -7,11 +7,13 @@
 #include "EdGraph/EdGraphSchema.h"
 #include "Layout/Visibility.h"
 #include "SGraphActionMenu.h"
+#include "WorkspaceAssetRegistryInfo.h"
 #include "Widgets/DeclarativeSyntaxSupport.h"
 #include "Widgets/Input/SEditableTextBox.h"
 #include "Widgets/Layout/SBorder.h"
 #include "Widgets/SWidget.h"
 
+struct FWorkspaceOutlinerItemExport;
 class UAnimNextRigVMAssetEditorData;
 class URigVMSchema;
 class SEditableTextBox;
@@ -29,6 +31,7 @@ struct FActionMenuContextData
 	TArray<UObject*> SelectedObjects;
 
 	UEdGraph* Graph = nullptr;
+	FWorkspaceOutlinerItemExport Export;
 	const URigVMSchema* RigVMSchema = nullptr;
 	URigVMHost* RigVMHost = nullptr;
 	IRigVMClientHost* RigVMClientHost = nullptr;
@@ -60,9 +63,9 @@ public:
 
 	SLATE_END_ARGS()
 
-	void Construct(const FArguments& InArgs, UEdGraph* InGraph);
+	void Construct(const FArguments& InArgs, UEdGraph* InGraph, const FWorkspaceOutlinerItemExport& InExport);
 
-	~SActionMenu();
+	virtual ~SActionMenu() override;
 
 	TSharedRef<SEditableTextBox> GetFilterTextBox();
 
diff --git a/Engine/Plugins/Experimental/AnimNext/Source/AnimNextEditor/Internal/Variables/VariablePickerArgs.h b/Engine/Plugins/Experimental/AnimNext/Source/AnimNextEditor/Internal/Variables/VariablePickerArgs.h
--- a/Engine/Plugins/Experimental/AnimNext/Source/AnimNextEditor/Internal/Variables/VariablePickerArgs.h
+++ b/Engine/Plugins/Experimental/AnimNext/Source/AnimNextEditor/Internal/Variables/VariablePickerArgs.h
@@ -2,6 +2,7 @@
 
 #pragma once
 
+#include "AnimNextExports.h"
 #include "Delegates/Delegate.h"
 #include "Variables/SAddVariablesDialog.h"
 
@@ -46,7 +47,10 @@ struct FVariablePickerArgs
 	// Delegate called when context sensitivity changes
 	FOnContextSensitivityChanged OnContextSensitivityChanged;
 
-	// Whether the search box should be focussed on widget creation
+	// Default filter for AnimNext export flags to be used for variable list population
+	EAnimNextExportedVariableFlags FlagFilter = EAnimNextExportedVariableFlags::Declared | EAnimNextExportedVariableFlags::Public;
+
+	// Whether the search box should be focused on widget creation
 	bool bFocusSearchWidget = true;
 };
 
diff --git a/Engine/Plugins/Experimental/AnimNext/Source/AnimNextEditor/Private/AnimNextEditorModule.cpp b/Engine/Plugins/Experimental/AnimNext/Source/AnimNextEditor/Private/AnimNextEditorModule.cpp
--- a/Engine/Plugins/Experimental/AnimNext/Source/AnimNextEditor/Private/AnimNextEditorModule.cpp
+++ b/Engine/Plugins/Experimental/AnimNext/Source/AnimNextEditor/Private/AnimNextEditorModule.cpp
@@ -6,7 +6,6 @@
 #include "EdGraphNode_Comment.h"
 #include "ISettingsModule.h"
 #include "IWorkspaceEditor.h"
-#include "ScopedTransaction.h"
 #include "SSimpleButton.h"
 #include "UncookedOnlyUtils.h"
 #include "Editor/RigVMEditorTools.h"
@@ -28,7 +27,6 @@
 #include "Editor/RigVMGraphDetailCustomization.h"
 #include "FileHelpers.h"
 #include "IUniversalObjectLocatorEditorModule.h"
-#include "MessageLogModule.h"
 #include "Param/AnimNextComponentLocatorEditor.h"
 #include "Param/AnimNextLocatorContext.h"
 #include "Param/ObjectCastLocatorEditor.h"
@@ -45,14 +43,12 @@
 #include "Common/AnimNextFunctionItemDetails.h"
 #include "Common/AnimNextRigVMAssetCommands.h"
 #include "Param/AnimNextActorLocatorEditor.h"
-#include "IWorkspaceEditor.h"
 #include "RigVMCommands.h"
 #include "ToolMenus.h"
 #include "Framework/Docking/LayoutExtender.h"
 #include "Common/AnimNextAssetItemDetails.h"
 #include "Common/AnimNextRigVMAssetEditorDataCustomization.h"
 #include "Variables/AnimNextSharedVariables.h"
-#include "Editor/RigVMEditorStyle.h"
 #include "Tools/EdModeInteractiveToolsContext.h"
 #include "Entries/AnimNextEventGraphEntry.h"
 #include "Logging/MessageLog.h"
@@ -61,7 +57,6 @@
 #include "Variables/SVariablesView.h"
 #include "Variables/VariableOverrideCommands.h"
 #include "Workspace/AnimNextWorkspaceEditorMode.h"
-#include "Workspace/AnimNextWorkspaceSchema.h"
 #include "ExternalPackageHelper.h"
 #include "Common/GraphEditorSchemaActions.h"
 #include "EditorUtils.h"
@@ -76,6 +71,11 @@
 #include "RewindDebugger/AnimNextModuleTrack.h"
 #include "RewindDebugger/AnimNextTraceModule.h"
 #include "Variables/VariableReferencePropertyCustomization.h"
+#include "AnimNextEditorStyle.h"
+#include "AnimNextSharedVariableNode.h"
+#include "Common/AnimNextAssetFindReplaceVariables.h"
+#include "Toolkits/AssetEditorToolkitMenuContext.h"
+#include "Variables/Outliner/VariablesOutlinerCommands.h"
 
 #define LOCTEXT_NAMESPACE "AnimNextEditorModule"
 
@@ -98,6 +98,8 @@ void FAnimNextEditorModule::StartupModule()
 	FVariableOverrideCommands::Register();
 	FRigVMCommands::Register();
 	FAnimNextRigVMAssetCommands::Register();
+	FVariablesOutlinerCommands::Register();
+	FAnimNextEditorStyle::Get();
 
 	// Register settings for user editing
 	ISettingsModule& SettingsModule = FModuleManager::Get().LoadModuleChecked<ISettingsModule>("Settings");
@@ -163,6 +165,10 @@ void FAnimNextEditorModule::StartupModule()
 			TSharedRef<FAnimNextVariablesTabSummoner> VariablesTabSummoner = MakeShared<FAnimNextVariablesTabSummoner>(InEditorPtr);
 			TabFactories.RegisterFactory(VariablesTabSummoner);
 			VariablesTabSummoner->RegisterTabSpawner(InTabManager, nullptr);
+
+			TSharedRef<FAnimNextFindAndReplaceTabSummoner> FindAndReplaceTabSummoner = MakeShared<FAnimNextFindAndReplaceTabSummoner>(InEditorPtr);
+			TabFactories.RegisterFactory(FindAndReplaceTabSummoner);
+			FindAndReplaceTabSummoner->RegisterTabSpawner(InTabManager, nullptr);
 		});
 
 	WorkspaceEditorModule.OnExtendTabs().AddLambda([](FLayoutExtender& InLayoutExtender, TSharedPtr<UE::Workspace::IWorkspaceEditor> InEditorPtr)
@@ -173,10 +179,108 @@ void FAnimNextEditorModule::StartupModule()
 		FTabManager::FTab FindTab(FTabId(FindTabName), ETabState::ClosedTab);
 		InLayoutExtender.ExtendLayout(FTabId(Workspace::WorkspaceTabs::BottomMiddleDocumentArea), ELayoutExtensionPosition::After, FindTab);
 
+		FTabManager::FTab FindAndReplaceTab(FTabId(FindAndReplaceTabName), ETabState::ClosedTab);
+		InLayoutExtender.ExtendLayout(FTabId(Workspace::WorkspaceTabs::BottomMiddleDocumentArea), ELayoutExtensionPosition::After, FindAndReplaceTab);
+
 		FTabManager::FTab VariablesTab(FTabId(VariablesTabName), ETabState::OpenedTab);
 		InLayoutExtender.ExtendLayout(FTabId(Workspace::WorkspaceTabs::BottomLeftDocumentArea), ELayoutExtensionPosition::After, VariablesTab);
 	});
 
+	UToolMenu* AnimNextEdGraphContextMenu = UToolMenus::Get()->ExtendMenu( "GraphEditor.GraphNodeContextMenu.AnimNextEdGraphNode");
+	AnimNextEdGraphContextMenu->AddDynamicSection("FindAndReplace", FNewToolMenuDelegate::CreateLambda([](UToolMenu* InMenu)
+	{
+		if (UGraphNodeContextMenuContext* GraphNodeContextMenuContext = InMenu->Context.FindContext<UGraphNodeContextMenuContext>())
+		{
+			if (const URigVMEdGraphNode* RigVMEdGraphNode = Cast<URigVMEdGraphNode>(GraphNodeContextMenuContext->Node))
+			{
+				if (UAssetEditorToolkitMenuContext* AssetEditorToolkitMenuContext = InMenu->Context.FindContext<UAssetEditorToolkitMenuContext>())
+				{
+					if (TSharedPtr<FAssetEditorToolkit> SharedToolkit = AssetEditorToolkitMenuContext->Toolkit.Pin())
+					{
+						if (SharedToolkit->GetToolkitFName() == Workspace::IWorkspaceEditor::GetWorkspaceEditorToolkitName())
+						{
+							FAnimNextSoftVariableReference SearchReference;
+							FAnimNextParamType SearchType;
+
+							if (UAnimNextSharedVariableNode* SharedVariableNode = Cast<UAnimNextSharedVariableNode>(RigVMEdGraphNode->GetModelNode()))
+							{
+								const UAnimNextSharedVariables* Asset = SharedVariableNode->Asset.Get();
+								const UScriptStruct* Struct = SharedVariableNode->Struct.Get();
+
+								check(SharedVariableNode->Type == EAnimNextSharedVariablesType::Asset && Asset || SharedVariableNode->Type == EAnimNextSharedVariablesType::Struct && Struct);
+
+								const UObject* Object = Asset ? Cast<UObject>(Asset) : Cast<UObject>(Struct);
+								SearchReference = FAnimNextSoftVariableReference(SharedVariableNode->GetVariableName(), Object);
+								if (SharedVariableNode->Type == EAnimNextSharedVariablesType::Asset)
+								{
+									const UAnimNextRigVMAssetEditorData* EditorData = UE::AnimNext::UncookedOnly::FUtils::GetEditorData<const UAnimNextRigVMAssetEditorData, const UAnimNextSharedVariables>(Asset);						
+									const UAnimNextVariableEntry* Entry = Cast<UAnimNextVariableEntry>(EditorData->FindEntry(SharedVariableNode->GetVariableName()));
+									SearchType = Entry->GetType();
+								}
+								else if (SharedVariableNode->Type == EAnimNextSharedVariablesType::Struct)
+								{
+									if (const FProperty* Property = SharedVariableNode->Struct->FindPropertyByName(SharedVariableNode->GetVariableName()))
+									{
+										SearchType = FAnimNextParamType::FromProperty(Property);
+									}
+								}
+							}
+							else if (URigVMVariableNode* VariableNode = Cast<URigVMVariableNode>(RigVMEdGraphNode->GetModelNode()))
+							{
+								if (const UAnimNextRigVMAsset* Asset = Cast<UAnimNextRigVMAsset>(VariableNode->GetTypedOuter(UAnimNextRigVMAsset::StaticClass())))
+								{
+									const UAnimNextRigVMAssetEditorData* EditorData = UncookedOnly::FUtils::GetEditorData<const UAnimNextRigVMAssetEditorData, const UAnimNextRigVMAsset>(Asset); 
+									const UAnimNextVariableEntry* Entry = Cast<UAnimNextVariableEntry>(EditorData->FindEntry(VariableNode->GetVariableName()));
+									SearchReference = FAnimNextSoftVariableReference(VariableNode->GetVariableName(), Asset);
+									SearchType = Entry->GetType();
+								}
+							}
+
+							if (!SearchReference.IsNone() && SearchType.IsValid())
+							{
+								FToolMenuSection& FindReplaceSection = InMenu->AddSection("FindAndReplace", LOCTEXT("FindReferencesLabel", "Find References"));			
+								FindReplaceSection.AddSubMenu("FindReferencesMenu", LOCTEXT("FindReferencesMenuLabel", "Find References..."), LOCTEXT("FindReferencesMenuToolTip", "Find references to selected variables by different means."),
+								FNewToolMenuDelegate::CreateLambda(
+									[SearchReference, SearchType](UToolMenu* InSubmenu)
+									{									
+										auto InvokeFindVariables = [SearchReference, SearchType](const FToolMenuContext& Context, ESearchScope Scope)
+										{
+											UAssetEditorToolkitMenuContext* AssetEditorToolkitMenuContext = Context.FindContext<UAssetEditorToolkitMenuContext>();										 
+											if (TSharedPtr<Workspace::IWorkspaceEditor> SharedWorkspaceEditor = StaticCastSharedPtr<Workspace::IWorkspaceEditor>(AssetEditorToolkitMenuContext->Toolkit.Pin()))
+											{
+												if (TSharedPtr<SDockTab> FindAndReplaceTab = SharedWorkspaceEditor->GetTabManager()->TryInvokeTab(FTabId("FindAndReplaceTab")))
+												{
+													TSharedRef<IAnimAssetFindReplace> FindAndReplace = StaticCastSharedRef<IAnimAssetFindReplace>(FindAndReplaceTab->GetContent());
+													FindAndReplace->SetCurrentProcessor(UAnimNextAssetFindReplaceVariables::StaticClass());
+													UAnimNextAssetFindReplaceVariables* AnimNextFindAndReplaceProcessor = FindAndReplace->GetProcessor<UAnimNextAssetFindReplaceVariables>();
+													
+													AnimNextFindAndReplaceProcessor->SetWorkspaceEditor(SharedWorkspaceEditor.ToSharedRef());
+													
+													AnimNextFindAndReplaceProcessor->SetSearchScope(Scope);
+													AnimNextFindAndReplaceProcessor->SetFindReference(SearchReference, SearchType);					
+												}
+											}
+										};
+										
+										 InSubmenu->AddMenuEntry(NAME_None, FToolMenuEntry::InitMenuEntry("FindGlobalReferences", LOCTEXT("GlobalScopeLabel", "In Project"), LOCTEXT("", ""),UAnimNextAssetFindReplaceVariables::GetIconFromSearchScope(ESearchScope::Global), FToolUIAction(FToolMenuExecuteAction::CreateLambda(InvokeFindVariables, ESearchScope::Global))));
+
+										InSubmenu->AddMenuEntry(NAME_None, FToolMenuEntry::InitMenuEntry("FindWorkspaceReferences", LOCTEXT("WorkspaceScopeLabel", "In Workspace"), LOCTEXT("", ""),UAnimNextAssetFindReplaceVariables::GetIconFromSearchScope(ESearchScope::Workspace), FToolUIAction(FToolMenuExecuteAction::CreateLambda(InvokeFindVariables, ESearchScope::Workspace))));
+
+										InSubmenu->AddMenuEntry(NAME_None, FToolMenuEntry::InitMenuEntry("FindAssetReferences", LOCTEXT("AssetScopeLabel", "In Asset"), LOCTEXT("", ""),UAnimNextAssetFindReplaceVariables::GetIconFromSearchScope(ESearchScope::Asset), FToolUIAction(FToolMenuExecuteAction::CreateLambda(InvokeFindVariables, ESearchScope::Asset))));		
+									}),
+									FToolUIActionChoice(),
+									EUserInterfaceActionType::Button,
+									false,
+									FSlateIcon(FAppStyle::GetAppStyleSetName(), TEXT("Icons.Search"))
+								);	
+							}				
+						}
+					}
+				}
+			}			
+		}
+	}));
+
 	WorkspaceEditorModule.OnExtendToolMenuContext().AddLambda([](TSharedPtr<UE::Workspace::IWorkspaceEditor> InWorkspaceEditor, FToolMenuContext& InContext)
 	{
 		UContextObjectStore* ContextStore = InWorkspaceEditor->GetEditorModeManager().GetInteractiveToolsContext()->ContextObjectStore;
@@ -187,7 +291,10 @@ void FAnimNextEditorModule::StartupModule()
 			AnimNextEditorContext->WeakWorkspaceEditor = InWorkspaceEditor;
 			ContextStore->AddContextObject(AnimNextEditorContext);
 		}
-		
+
+		InContext.AddObject(AnimNextEditorContext);
+
+		// [TODO] This should be a dynamic extension
 		if (UToolMenu* Menu = UToolMenus::Get()->ExtendMenu(InWorkspaceEditor->GetToolMenuToolbarName()))
 		{
 			FToolMenuSection& RigVMOperationsSection = Menu->AddSection("RigVMOperations", TAttribute<FText>(), FToolMenuInsert("WorkspaceOperations", EToolMenuInsertType::After));
@@ -358,6 +465,7 @@ void FAnimNextEditorModule::StartupModule()
 
 	const TSharedPtr<FAnimNextAssetItemDetails> AssetItemDetails = MakeShared<FAnimNextAssetItemDetails>();
 	WorkspaceModule.RegisterWorkspaceItemDetails(Workspace::FOutlinerItemDetailsId(FAnimNextModuleOutlinerData::StaticStruct()->GetFName()), AssetItemDetails);
+	WorkspaceModule.RegisterWorkspaceItemDetails(Workspace::FOutlinerItemDetailsId(FAnimNextSharedVariablesOutlinerData::StaticStruct()->GetFName()), AssetItemDetails);
 
 	FAnimNextGraphItemDetails::RegisterToolMenuExtensions();
 	FAnimNextAssetItemDetails::RegisterToolMenuExtensions();
@@ -504,14 +612,14 @@ void FAnimNextEditorModule::UnregisterNodeDblClickHandler(const FDelegateHandle&
 	OnNodeDblClickHandlerMulticast.Remove(InDelegateHandle);
 }
 
-void FAnimNextEditorModule::CollectGraphMenuActions(const TWeakPtr<UE::Workspace::IWorkspaceEditor>& WorkspaceEditor, FGraphContextMenuBuilder& InContexMenuBuilder, const FActionMenuContextData& InActionMenuContextData)
+void FAnimNextEditorModule::CollectGraphMenuActions(const TWeakPtr<UE::Workspace::IWorkspaceEditor>& WorkspaceEditor, FGraphContextMenuBuilder& InContextMenuBuilder, const FActionMenuContextData& InActionMenuContextData)
 {
 	if (OnCollectGraphMenuActionsDelegateImpl.IsBound())
 	{
-		OnCollectGraphMenuActionsDelegateImpl.Broadcast(WorkspaceEditor, InContexMenuBuilder, InActionMenuContextData);
+		OnCollectGraphMenuActionsDelegateImpl.Broadcast(WorkspaceEditor, InContextMenuBuilder, InActionMenuContextData);
 	}
 
-	if (const URigVMEdGraph* RigVMEdGraph = Cast<URigVMEdGraph>(InContexMenuBuilder.CurrentGraph))
+	if (const URigVMEdGraph* RigVMEdGraph = Cast<URigVMEdGraph>(InContextMenuBuilder.CurrentGraph))
 	{
 		const URigVMGraph* Graph = RigVMEdGraph->GetModel();
 		check(Graph);
@@ -559,7 +667,7 @@ void FAnimNextEditorModule::CollectGraphMenuActions(const TWeakPtr<UE::Workspace
 				continue;
 			}
 
-			UE::AnimNext::Editor::FUtils::AddSchemaRigUnitAction(URigVMUnitNode::StaticClass(), Struct, Function, InContexMenuBuilder);
+			UE::AnimNext::Editor::FUtils::AddSchemaRigUnitAction(URigVMUnitNode::StaticClass(), Struct, Function, InContextMenuBuilder);
 		}
 
 		for (const FRigVMDispatchFactory* Factory : FRigVMRegistry::Get().GetFactories())
@@ -591,7 +699,7 @@ void FAnimNextEditorModule::CollectGraphMenuActions(const TWeakPtr<UE::Workspace
 			FText MenuDesc = FText::FromString(Factory->GetNodeTitle(FRigVMTemplateTypeMap()));
 			FText ToolTip = Factory->GetNodeTooltip(FRigVMTemplateTypeMap());
 
-			InContexMenuBuilder.AddAction(MakeShared<FAnimNextSchemaAction_DispatchFactory>(Template->GetNotation(), NodeCategory, MenuDesc, ToolTip));
+			InContextMenuBuilder.AddAction(MakeShared<FAnimNextSchemaAction_DispatchFactory>(Template->GetNotation(), NodeCategory, MenuDesc, ToolTip));
 		};
 
 		if (URigVMFunctionLibrary* LocalFunctionLibrary = RigVMClientHost->GetLocalFunctionLibrary())
@@ -609,7 +717,7 @@ void FAnimNextEditorModule::CollectGraphMenuActions(const TWeakPtr<UE::Workspace
 				const FText MenuDesc = FText::FromString(FunctionLibraryNode->GetName());
 				const FText ToolTip = FunctionLibraryNode->GetToolTipText();
 
-				InContexMenuBuilder.AddAction(MakeShared<FAnimNextSchemaAction_Function>(FunctionLibraryNode, NodeCategory, MenuDesc, ToolTip));
+				InContextMenuBuilder.AddAction(MakeShared<FAnimNextSchemaAction_Function>(FunctionLibraryNode, NodeCategory, MenuDesc, ToolTip));
 			}
 		}
 
@@ -629,39 +737,101 @@ void FAnimNextEditorModule::CollectGraphMenuActions(const TWeakPtr<UE::Workspace
 					const FText MenuDesc = FText::FromString(FunctionHeader.NodeTitle);
 					const FText ToolTip = FunctionHeader.GetTooltip();
 
-					InContexMenuBuilder.AddAction(MakeShared<FAnimNextSchemaAction_Function>(FunctionHeader, NodeCategory, MenuDesc, ToolTip));
+					InContextMenuBuilder.AddAction(MakeShared<FAnimNextSchemaAction_Function>(FunctionHeader, NodeCategory, MenuDesc, ToolTip));
 				}
 			}
 		}
 
 		TArray<UAnimNextRigVMAssetEditorData::FVariableInfo> VariableInfoArray;
 		EditorData->GetAllVariables(VariableInfoArray, UAnimNextRigVMAssetEditorData::EVariableRecursion::IncludeShared, UAnimNextRigVMAssetEditorData::EVariableAccessFilter::All);
+
+		TArray<const UAnimNextRigVMAssetEditorData::FVariableInfo> ProcessedVariables;
+		
 		for (const UAnimNextRigVMAssetEditorData::FVariableInfo& VariableInfo : VariableInfoArray)
 		{
-			InContexMenuBuilder.AddAction(MakeShared<FAnimNextSchemaAction_Variable>(VariableInfo.Name, VariableInfo.SourceObject, VariableInfo.Type, FAnimNextSchemaAction_Variable::EVariableAccessorChoice::Set));
-			InContexMenuBuilder.AddAction(MakeShared<FAnimNextSchemaAction_Variable>(VariableInfo.Name, VariableInfo.SourceObject, VariableInfo.Type, FAnimNextSchemaAction_Variable::EVariableAccessorChoice::Get));
+			InContextMenuBuilder.AddAction(MakeShared<FAnimNextSchemaAction_Variable>(VariableInfo.Name, VariableInfo.SourceObject, VariableInfo.Type, FAnimNextSchemaAction_Variable::EVariableAccessorChoice::Set));
+			InContextMenuBuilder.AddAction(MakeShared<FAnimNextSchemaAction_Variable>(VariableInfo.Name, VariableInfo.SourceObject, VariableInfo.Type, FAnimNextSchemaAction_Variable::EVariableAccessorChoice::Get));
+
+			ProcessedVariables.AddUnique(VariableInfo);
 		}
 
-		InContexMenuBuilder.AddAction(MakeShared<FAnimNextSchemaAction_AddComment>());
-	}
+		InContextMenuBuilder.AddAction(MakeShared<FAnimNextSchemaAction_AddComment>());
+
+		// Collect all variables available within the current working context (asset export chain recursing upwards to module)
+		if(FWorkspaceOutlinerAssetReferenceItemData::IsAssetReference(InActionMenuContextData.Export))
+		{			
+			TArray<FSoftObjectPath> AssetPaths;
+			InActionMenuContextData.Export.GetAssetPaths(AssetPaths);
+		
+			for (const FSoftObjectPath& SoftObjectPath : AssetPaths)
+			{
+				if(UAnimNextRigVMAsset* RigVMAsset = Cast<UAnimNextRigVMAsset>(SoftObjectPath.TryLoad()))
+				{
+					if(UAnimNextRigVMAssetEditorData* AssetEditorData = UncookedOnly::FUtils::GetEditorData(RigVMAsset))
+					{
+						// Skip editor-data if its already been handled
+						if (AssetEditorData != EditorData)
+						{
+							VariableInfoArray.Reset();
+							AssetEditorData->GetAllVariables(VariableInfoArray, UAnimNextRigVMAssetEditorData::EVariableRecursion::IncludeShared, UAnimNextRigVMAssetEditorData::EVariableAccessFilter::All);
+							for (const UAnimNextRigVMAssetEditorData::FVariableInfo& VariableInfo : VariableInfoArray)
+							{
+								if (!ProcessedVariables.Contains(VariableInfo))
+								{
+									InContextMenuBuilder.AddAction(MakeShared<FAnimNextSchemaAction_Variable>(VariableInfo.Name, VariableInfo.SourceObject, VariableInfo.Type, FAnimNextSchemaAction_Variable::EVariableAccessorChoice::Set));
+									InContextMenuBuilder.AddAction(MakeShared<FAnimNextSchemaAction_Variable>(VariableInfo.Name, VariableInfo.SourceObject, VariableInfo.Type, FAnimNextSchemaAction_Variable::EVariableAccessorChoice::Get));
+								}
+							}
+						}
+					}
+				}
+			}
+		}
+
+		// Promotion to variable
+		if (InContextMenuBuilder.FromPin != nullptr)
+		{
+			const FAnimNextParamType ParamType = UE::AnimNext::UncookedOnly::FUtils::GetParamTypeFromPinType(InContextMenuBuilder.FromPin->PinType);
+			if (ParamType.IsValid())
+			{
+				InContextMenuBuilder.AddAction(MakeShared<FAnimNextSchemaAction_PromoteToVariable>());
+			}
+		}
+	}	
 }
 
 void FAnimNextEditorModule::RegisterWorkspaceDocumentTypes(Workspace::IWorkspaceEditorModule& WorkspaceEditorModule)
 {
 	// --- AnimNextModule ---
 	Workspace::FObjectDocumentArgs AnimNextModuleDocumentArgs(
-			Workspace::FOnRedirectWorkspaceContext::CreateLambda([](UObject* InObject) -> UObject*
+			Workspace::FOnRedirectWorkspaceContext::CreateLambda([](const Workspace::FWorkspaceDocument& Document) -> UObject*
 			{
-				UAnimNextModule* Module = CastChecked<UAnimNextModule>(InObject);
-				UAnimNextModule_EditorData* EditorData = UncookedOnly::FUtils::GetEditorData<UAnimNextModule_EditorData>(Module);
-
-				// Redirect to the inner graph, if any
-				UAnimNextEventGraphEntry* EventGraphEntry = EditorData->FindFirstEntryOfType<UAnimNextEventGraphEntry>();
-				if(EventGraphEntry)
+				UObject* InObject = Document.GetObject();
+				if(UAnimNextModule* Module = Cast<UAnimNextModule>(InObject))
 				{
-					return EventGraphEntry->GetEdGraph();
+				    UAnimNextModule_EditorData* EditorData = UncookedOnly::FUtils::GetEditorData<UAnimNextModule_EditorData>(Module);
+
+					const TInstancedStruct<FWorkspaceOutlinerItemData>& ExportData = Document.Export.GetData();
+					if (ExportData.GetScriptStruct() == FAnimNextGraphOutlinerData::StaticStruct())
+					{
+						const FAnimNextGraphOutlinerData& TypeData = ExportData.Get<FAnimNextGraphOutlinerData>();
+						if (UAnimNextEventGraphEntry* Entry = Cast<UAnimNextEventGraphEntry>(TypeData.GetEntry()))
+						{
+							return Entry->GetEdGraph();
+						}
+					}
+					
+				    // Redirect to the inner graph, if any
+				    UAnimNextEventGraphEntry* EventGraphEntry = EditorData->FindFirstEntryOfType<UAnimNextEventGraphEntry>();
+				    if(EventGraphEntry)
+				    {
+					    return EventGraphEntry->GetEdGraph();
+				    }
+    
+				    return nullptr;
 				}
-				return nullptr;
+				
+				return InObject;
 			}));
 
 	AnimNextModuleDocumentArgs.DocumentEditorMode = UAnimNextWorkspaceEditorMode::EM_AnimNextWorkspace;
@@ -676,14 +846,14 @@ void FAnimNextEditorModule::RegisterWorkspaceDocumentTypes(Workspace::IWorkspace
 		const FString WorkspaceAssetPath = FExternalPackageHelper::GetExternalObjectsPath(InContext.WorkspaceEditor->GetPackageName());
 		TWeakPtr<UE::Workspace::IWorkspaceEditor> WorkspaceEditorWeak = InContext.WorkspaceEditor.ToWeakPtr();
 
-		TSharedRef<SActionMenu> ActionMenu = SNew(SActionMenu, InGraph)
+		TSharedRef<SActionMenu> ActionMenu = SNew(SActionMenu, InGraph, InContext.Document.Export)
 			.AutoExpandActionMenu(bAutoExpand)
 			.NewNodePosition(InNodePosition)
 			.DraggedFromPins(InDraggedPins)
 			.OnClosedCallback(InOnMenuClosed)
-			.OnCollectGraphActionsCallback(SActionMenu::FCollectAllGraphActions::CreateLambda([this, WorkspaceEditorWeak](FGraphContextMenuBuilder& InContexMenuBuilder, const FActionMenuContextData& InActionMenuContextData)
+			.OnCollectGraphActionsCallback(SActionMenu::FCollectAllGraphActions::CreateLambda([this, WorkspaceEditorWeak](FGraphContextMenuBuilder& InContextMenuBuilder, const FActionMenuContextData& InActionMenuContextData)
 				{
-					CollectGraphMenuActions(WorkspaceEditorWeak, InContexMenuBuilder, InActionMenuContextData);
+					CollectGraphMenuActions(WorkspaceEditorWeak, InContextMenuBuilder, InActionMenuContextData);
 				}));
 
 		TSharedPtr<SWidget> FilterTextBox = StaticCastSharedRef<SWidget>(ActionMenu->GetFilterTextBox());
@@ -808,7 +978,7 @@ void FAnimNextEditorModule::RegisterWorkspaceDocumentTypes(Workspace::IWorkspace
 			return;
 		}
 
-		URigVMEdGraph* RigVMEdGraph = Cast<URigVMEdGraph>(InContext.Document.Object);
+		URigVMEdGraph* RigVMEdGraph = InContext.Document.GetTypedObject<URigVMEdGraph>();
 		if (RigVMEdGraph == nullptr)
 		{
 			return;
@@ -837,7 +1007,7 @@ void FAnimNextEditorModule::RegisterWorkspaceDocumentTypes(Workspace::IWorkspace
 			return;
 		}
 
-		URigVMEdGraph* RigVMEdGraph = Cast<URigVMEdGraph>(InContext.Document.Object);
+		URigVMEdGraph* RigVMEdGraph = InContext.Document.GetTypedObject<URigVMEdGraph>();
 		if (RigVMEdGraph == nullptr)
 		{
 			return;
@@ -883,7 +1053,7 @@ void FAnimNextEditorModule::RegisterWorkspaceDocumentTypes(Workspace::IWorkspace
 				return;
 			}
 
-			URigVMEdGraph* RigVMEdGraph = Cast<URigVMEdGraph>(InContext.Document.Object);
+			URigVMEdGraph* RigVMEdGraph = InContext.Document.GetTypedObject<URigVMEdGraph>();
 			if (RigVMEdGraph == nullptr)
 			{
 				return;
@@ -924,7 +1094,7 @@ void FAnimNextEditorModule::RegisterWorkspaceDocumentTypes(Workspace::IWorkspace
 				return;
 			}
 
-			URigVMEdGraph* RigVMEdGraph = Cast<URigVMEdGraph>(InContext.Document.Object);
+			URigVMEdGraph* RigVMEdGraph = InContext.Document.GetTypedObject<URigVMEdGraph>();
 			if (RigVMEdGraph == nullptr)
 			{
 				return;
@@ -946,7 +1116,7 @@ void FAnimNextEditorModule::RegisterWorkspaceDocumentTypes(Workspace::IWorkspace
 		});
 	GraphArgs.OnGraphSelectionChanged = Workspace::FOnGraphSelectionChanged::CreateLambda([this](const Workspace::FWorkspaceEditorContext& InContext, const FGraphPanelSelectionSet& NewSelection)
 	{
-		URigVMEdGraph* RigVMEdGraph = Cast<URigVMEdGraph>(InContext.Document.Object);
+		URigVMEdGraph* RigVMEdGraph = InContext.Document.GetTypedObject<URigVMEdGraph>();
 		if (RigVMEdGraph == nullptr)
 		{
 			return;
@@ -1021,7 +1191,7 @@ void FAnimNextEditorModule::RegisterWorkspaceDocumentTypes(Workspace::IWorkspace
 	{
 		TWeakPtr<Workspace::IWorkspaceEditor> WeakWorkspaceEditor = InContext.WorkspaceEditor;
 
-		if (UAnimNextEdGraph* EdGraph = Cast<UAnimNextEdGraph>(InContext.Document.Object))
+		if (UAnimNextEdGraph* EdGraph = InContext.Document.GetTypedObject<UAnimNextEdGraph>())
 		{
 			UAnimNextRigVMAssetEditorData* EditorData = EdGraph->GetTypedOuter<UAnimNextRigVMAssetEditorData>();
 			if(EditorData)
@@ -1059,7 +1229,6 @@ void FAnimNextEditorModule::RegisterWorkspaceDocumentTypes(Workspace::IWorkspace
 				if (TSharedPtr<Workspace::IWorkspaceEditor> WorkspaceEditor = WeakWorkspaceEditor.Pin())
 				{
 					const TSharedRef<FUICommandList>& CommandList = WorkspaceEditor->GetToolkitCommands();
-
 					{
 						const FAnimNextRigVMAssetCommands& Commands = FAnimNextRigVMAssetCommands::Get();
 
@@ -1077,6 +1246,21 @@ void FAnimNextEditorModule::RegisterWorkspaceDocumentTypes(Workspace::IWorkspace
 
 						CommandList->MapAction(Commands.FindInAnimNextRigVMAsset,
 							FExecuteAction::CreateLambda(TryFindAnimNextRigVMAsset));
+
+						auto TryFindAndReplaceAnimNextRigVMAsset = [WeakWorkspaceEditor]()
+						{
+							if (TSharedPtr<Workspace::IWorkspaceEditor> WorkspaceEditorPinned = WeakWorkspaceEditor.Pin())
+							{
+								if (TSharedPtr<SDockTab> Tab = WorkspaceEditorPinned->GetTabManager()->TryInvokeTab(FTabId(FindAndReplaceTabName)))
+								{
+									TSharedRef<IAnimAssetFindReplace> FindInAnimNextRigVMAssetWidget = StaticCastSharedRef<IAnimAssetFindReplace>(Tab->GetContent());
+									Tab->DrawAttention();
+								}
+							}
+						};
+
+						CommandList->MapAction(Commands.FindAndReplaceInAnimNextRigVMAsset,
+							FExecuteAction::CreateLambda(TryFindAndReplaceAnimNextRigVMAsset));
 					}
 				}
 			}
@@ -1092,7 +1276,7 @@ void FAnimNextEditorModule::RegisterWorkspaceDocumentTypes(Workspace::IWorkspace
 
 	GraphDocumentArgs.OnGetDocumentBreadcrumbTrail = Workspace::FOnGetDocumentBreadcrumbTrail::CreateLambda([](const Workspace::FWorkspaceEditorContext& InContext, TArray<TSharedPtr<Workspace::FWorkspaceBreadcrumb>>& OutBreadcrumbs)
 	{
-		if (URigVMEdGraph* EdGraph = Cast<URigVMEdGraph>(InContext.Document.Object))
+		if (URigVMEdGraph* EdGraph = InContext.Document.GetTypedObject<UAnimNextEdGraph>())
 		{
 			if (UAnimNextRigVMAssetEditorData* EditorData = EdGraph->GetTypedOuter<UAnimNextRigVMAssetEditorData>())
 			{
diff --git a/Engine/Plugins/Experimental/AnimNext/Source/AnimNextEditor/Private/AnimNextEditorModule.h b/Engine/Plugins/Experimental/AnimNext/Source/AnimNextEditor/Private/AnimNextEditorModule.h
--- a/Engine/Plugins/Experimental/AnimNext/Source/AnimNextEditor/Private/AnimNextEditorModule.h
+++ b/Engine/Plugins/Experimental/AnimNext/Source/AnimNextEditor/Private/AnimNextEditorModule.h
@@ -47,7 +47,7 @@ private:
 	virtual void UnregisterNodeDblClickHandler(const FDelegateHandle& InDelegateHandle) override;
 
 private:
-	void CollectGraphMenuActions(const TWeakPtr<UE::Workspace::IWorkspaceEditor>& WorkspaceEditor, FGraphContextMenuBuilder& InContexMenuBuilder, const FActionMenuContextData& InActionMenuContextData);
+	void CollectGraphMenuActions(const TWeakPtr<UE::Workspace::IWorkspaceEditor>& WorkspaceEditor, FGraphContextMenuBuilder& InContextMenuBuilder, const FActionMenuContextData& InActionMenuContextData);
 
 	void RegisterWorkspaceDocumentTypes(Workspace::IWorkspaceEditorModule& WorkspaceEditorModule);
 	void UnregisterWorkspaceDocumentTypes();
diff --git a/Engine/Plugins/Experimental/AnimNext/Source/AnimNextEditor/Private/AnimNextEditorStyle.h b/Engine/Plugins/Experimental/AnimNext/Source/AnimNextEditor/Private/AnimNextEditorStyle.h
new file mode 100644
--- /dev/null
+++ b/Engine/Plugins/Experimental/AnimNext/Source/AnimNextEditor/Private/AnimNextEditorStyle.h
@@ -0,0 +1,60 @@
+// Copyright Epic Games, Inc. All Rights Reserved.
+
+#pragma once
+
+#include "Styling/SlateStyle.h"
+#include "Styling/SlateStyleRegistry.h"
+#include "Styling/SlateStyleMacros.h"
+#include "Styling/AppStyle.h"
+#include "Styling/CoreStyle.h"
+#include "Styling/SlateTypes.h"
+#include "Interfaces/IPluginManager.h"
+
+#define IMAGE_PLUGIN_BRUSH_SVG( RelativePath, ... ) FSlateVectorImageBrush( FAnimNextEditorStyle::InContent(RelativePath, ".svg"), __VA_ARGS__)
+#define PLUGIN_BOX_BRUSH_BRUSH( RelativePath, ... ) FSlateBoxBrush( FAnimNextEditorStyle::InContent(RelativePath, ".png"), __VA_ARGS__)
+
+class FAnimNextEditorStyle final : public FSlateStyleSet
+{
+public:	
+	FAnimNextEditorStyle() : FSlateStyleSet("AnimNextEditorStyle")
+	{
+		const FVector2D Icon16x16(16.0f, 16.0f);
+		
+		SetContentRoot(FPaths::EngineContentDir() / TEXT("Editor/Slate"));
+		SetCoreContentRoot(FPaths::EngineContentDir() / TEXT("Slate"));
+
+		// Class Icons
+		{
+			Set("ClassIcon.AnimNextSharedVariables", new IMAGE_PLUGIN_BRUSH_SVG("Slate/Archive_16", Icon16x16));
+			Set("ClassIcon.AnimNextModule", new IMAGE_BRUSH_SVG( "Starship/Common/Modules", Icon16x16 ) );
+			Set("ClassIcon.AnimNextAnimationGraph", new IMAGE_BRUSH_SVG("Starship/Blueprints/icon_BlueprintEditor_EventGraph", Icon16x16));
+			Set("ClassIcon.AnimNextStateTree", new IMAGE_BRUSH_SVG(TEXT("UMG/TreeView"), Icon16x16));
+		}
+
+		Set("VariablesOutliner.SharedVariablesPill.Border", new PLUGIN_BOX_BRUSH_BRUSH("Slate/NamespaceBorder", FMargin(4.0f / 16.0f)));
+
+		const FTextBlockStyle NormalText = FAppStyle::GetWidgetStyle<FTextBlockStyle>("NormalText");
+		Set("VariablesOutliner.SharedVariablesPill.Text", FTextBlockStyle(NormalText)
+			.SetFont(DEFAULT_FONT("Regular", 8))
+			.SetColorAndOpacity(FSlateColor::UseForeground()));
+
+		FSlateStyleRegistry::RegisterSlateStyle(*this);
+	}
+
+	static FAnimNextEditorStyle& Get()
+	{ 
+		static FAnimNextEditorStyle Instance;
+		return Instance;
+	}
+	
+	static FString InContent(const FString& RelativePath, const ANSICHAR* Extension)
+	{
+		static FString ContentDir = IPluginManager::Get().FindPlugin(TEXT("AnimNext"))->GetContentDir();
+		return (ContentDir / RelativePath) + Extension;
+	}
+
+	virtual ~FAnimNextEditorStyle() override
+	{
+		FSlateStyleRegistry::UnRegisterSlateStyle(*this);
+	}
+};
\ No newline at end of file
diff --git a/Engine/Plugins/Experimental/AnimNext/Source/AnimNextEditor/Private/Common/AnimNextFindTabSummoner.cpp b/Engine/Plugins/Experimental/AnimNext/Source/AnimNextEditor/Private/Common/AnimNextFindTabSummoner.cpp
--- a/Engine/Plugins/Experimental/AnimNext/Source/AnimNextEditor/Private/Common/AnimNextFindTabSummoner.cpp
+++ b/Engine/Plugins/Experimental/AnimNext/Source/AnimNextEditor/Private/Common/AnimNextFindTabSummoner.cpp
@@ -6,8 +6,10 @@
 #include "FindInAnimNextRigVMAsset.h"
 #include "IAnimNextEditorModule.h"
 #include "IWorkspaceEditor.h"
-#include "MessageLogModule.h"
 #include "Modules/ModuleManager.h"
+#include "AnimAssetFindReplace.h"
+#include "AnimNextAssetFindReplaceVariables.h"
+#include "PersonaModule.h"
 
 #define LOCTEXT_NAMESPACE "WorkspaceTabSummoner"
 
@@ -36,6 +38,37 @@ FText FAnimNextFindTabSummoner::GetTabToolTipText(const FWorkflowTabSpawnInfo& I
 	return ViewMenuTooltip;
 }
 
+FAnimNextFindAndReplaceTabSummoner::FAnimNextFindAndReplaceTabSummoner(TSharedPtr<UE::Workspace::IWorkspaceEditor> InHostingApp) 
+	: FWorkflowTabFactory(UE::AnimNext::Editor::FindAndReplaceTabName, StaticCastSharedPtr<FAssetEditorToolkit>(InHostingApp))
+{
+	TabLabel = LOCTEXT("AnimNextFindAndReplaceTabLabel", "Find and Replace");
+	TabIcon = FSlateIcon(FAppStyle::GetAppStyleSetName(), "Kismet.Tabs.FindResults");
+	ViewMenuDescription = LOCTEXT("AnimNextFindTabMenuDescription", "Find and Replace");
+	ViewMenuTooltip = LOCTEXT("AnimNextFindAndReplaceTabToolTip", "Search and replace contents of Animation(Next) Assets.");
+	bIsSingleton = true;
+
+
+	FPersonaModule& PersonaModule = FModuleManager::Get().LoadModuleChecked<FPersonaModule>("Persona");
+
+	FAnimAssetFindReplaceConfig Config;
+	Config.InitialProcessorClass = UAnimNextAssetFindReplaceVariables::StaticClass();
+	FindAndReplaceWidget = StaticCastSharedRef<IAnimAssetFindReplace>(PersonaModule.CreateFindReplaceWidget(Config));
+
+	if (UAnimNextAssetFindReplaceVariables* Processor = FindAndReplaceWidget->GetProcessor<UAnimNextAssetFindReplaceVariables>())
+	{
+		Processor->SetWorkspaceEditor(InHostingApp.ToSharedRef());
+	}	
+}
+
+TSharedRef<SWidget> FAnimNextFindAndReplaceTabSummoner::CreateTabBody(const FWorkflowTabSpawnInfo& Info) const
+{
+	return FindAndReplaceWidget.ToSharedRef();
+}
+
+FText FAnimNextFindAndReplaceTabSummoner::GetTabToolTipText(const FWorkflowTabSpawnInfo& Info) const
+{
+	return ViewMenuTooltip;
+}
 } // end namespace UE::AnimNext::Editor
 
 #undef LOCTEXT_NAMESPACE
diff --git a/Engine/Plugins/Experimental/AnimNext/Source/AnimNextEditor/Private/Common/AnimNextFindTabSummoner.h b/Engine/Plugins/Experimental/AnimNext/Source/AnimNextEditor/Private/Common/AnimNextFindTabSummoner.h
--- a/Engine/Plugins/Experimental/AnimNext/Source/AnimNextEditor/Private/Common/AnimNextFindTabSummoner.h
+++ b/Engine/Plugins/Experimental/AnimNext/Source/AnimNextEditor/Private/Common/AnimNextFindTabSummoner.h
@@ -15,6 +15,8 @@ namespace UE::AnimNext::Editor
 	class SFindInAnimNextRigVMAsset;
 }	// end namespace UE::AnimNext::Editor
 
+class IAnimAssetFindReplace;
+
 namespace UE::AnimNext::Editor
 {
 struct FAnimNextFindTabSummoner : public FWorkflowTabFactory
@@ -31,4 +33,18 @@ private:
 	TSharedPtr<SFindInAnimNextRigVMAsset> AnimNextFindResultsWidget;
 };
 
+struct FAnimNextFindAndReplaceTabSummoner : public FWorkflowTabFactory
+{
+public:
+	FAnimNextFindAndReplaceTabSummoner(TSharedPtr<UE::Workspace::IWorkspaceEditor> InHostingApp);
+
+private:
+	// FWorkflowTabFactory interface
+	virtual TSharedRef<SWidget> CreateTabBody(const FWorkflowTabSpawnInfo& Info) const override;
+	virtual FText GetTabToolTipText(const FWorkflowTabSpawnInfo& Info) const override;
+
+	// The widget this tab spawner wraps
+	TSharedPtr<IAnimAssetFindReplace> FindAndReplaceWidget;
+};
+
 } // end namespace UE::AnimNext::Editor
diff --git a/Engine/Plugins/Experimental/AnimNext/Source/AnimNextEditor/Private/Common/AnimNextGraphItemDetails.cpp b/Engine/Plugins/Experimental/AnimNext/Source/AnimNextEditor/Private/Common/AnimNextGraphItemDetails.cpp
--- a/Engine/Plugins/Experimental/AnimNext/Source/AnimNextEditor/Private/Common/AnimNextGraphItemDetails.cpp
+++ b/Engine/Plugins/Experimental/AnimNext/Source/AnimNextEditor/Private/Common/AnimNextGraphItemDetails.cpp
@@ -38,9 +38,9 @@ bool FAnimNextGraphItemDetails::HandleDoubleClick(const FToolMenuContext& ToolMe
 					{
 						if(const IRigVMClientHost* RigVMClientHost = RigVMGraph->GetImplementingOuter<IRigVMClientHost>())
 						{
-							if(UObject* EditorObject = RigVMClientHost->GetEditorObjectForRigVMGraph(RigVMGraph))
+							if(RigVMClientHost->GetEditorObjectForRigVMGraph(RigVMGraph))
 							{
-								WorkspaceEditor->OpenObjects({EditorObject});
+								WorkspaceEditor->OpenExports({WorkspaceItemContext->SelectedExports[0].GetResolvedExport()});
 								return true;
 							}
 						}
diff --git a/Engine/Plugins/Experimental/AnimNext/Source/AnimNextEditor/Private/Common/AnimNextRigVMAssetCommands.cpp b/Engine/Plugins/Experimental/AnimNext/Source/AnimNextEditor/Private/Common/AnimNextRigVMAssetCommands.cpp
--- a/Engine/Plugins/Experimental/AnimNext/Source/AnimNextEditor/Private/Common/AnimNextRigVMAssetCommands.cpp
+++ b/Engine/Plugins/Experimental/AnimNext/Source/AnimNextEditor/Private/Common/AnimNextRigVMAssetCommands.cpp
@@ -15,6 +15,7 @@ FAnimNextRigVMAssetCommands::FAnimNextRigVMAssetCommands()
 void FAnimNextRigVMAssetCommands::RegisterCommands()
 {
 	UI_COMMAND(FindInAnimNextRigVMAsset, "FindInUAFRigVMAsset", "Search the current UAF Asset.", EUserInterfaceActionType::Button, FInputChord(EModifierKey::Control, EKeys::F));
+	UI_COMMAND(FindAndReplaceInAnimNextRigVMAsset, "FindAndReplaceInUAFRigVMAsset", "Find and Replace across UAF Assets.", EUserInterfaceActionType::Button, FInputChord(EModifierKey::Control | EModifierKey::Shift, EKeys::F));
 }
 
 }
diff --git a/Engine/Plugins/Experimental/AnimNext/Source/AnimNextEditor/Private/Common/AnimNextRigVMAssetCommands.h b/Engine/Plugins/Experimental/AnimNext/Source/AnimNextEditor/Private/Common/AnimNextRigVMAssetCommands.h
--- a/Engine/Plugins/Experimental/AnimNext/Source/AnimNextEditor/Private/Common/AnimNextRigVMAssetCommands.h
+++ b/Engine/Plugins/Experimental/AnimNext/Source/AnimNextEditor/Private/Common/AnimNextRigVMAssetCommands.h
@@ -16,6 +16,8 @@ public:
 	virtual void RegisterCommands() override;
 
 	TSharedPtr<FUICommandInfo> FindInAnimNextRigVMAsset;
+
+	TSharedPtr<FUICommandInfo> FindAndReplaceInAnimNextRigVMAsset;
 };
 
 }
\ No newline at end of file
diff --git a/Engine/Plugins/Experimental/AnimNext/Source/AnimNextEditor/Private/Common/FindInAnimNextRigVMAsset.cpp b/Engine/Plugins/Experimental/AnimNext/Source/AnimNextEditor/Private/Common/FindInAnimNextRigVMAsset.cpp
--- a/Engine/Plugins/Experimental/AnimNext/Source/AnimNextEditor/Private/Common/FindInAnimNextRigVMAsset.cpp
+++ b/Engine/Plugins/Experimental/AnimNext/Source/AnimNextEditor/Private/Common/FindInAnimNextRigVMAsset.cpp
@@ -7,6 +7,7 @@
 #include "AnimNextRigVMAssetEditorData.h"
 #include "EdGraph/RigVMEdGraph.h"
 #include "GraphEditor.h"
+#include "IWorkspaceEditorModule.h"
 #include "RigVMModel/RigVMGraph.h"
 #include "UncookedOnlyUtils.h"
 
@@ -63,13 +64,13 @@ void SFindInAnimNextRigVMAsset::MatchTokens(const TArray<FString>& Tokens)
 
 	if (TSharedPtr<UE::Workspace::IWorkspaceEditor> WorkspaceEditor = StaticCastSharedPtr<UE::Workspace::IWorkspaceEditor>(AssetEditorToolkitPtr.Pin()))
 	{
-		TObjectPtr<UObject> FocusedDocument = WorkspaceEditor->GetFocussedDocument();
-		if (UEdGraph* Graph = Cast<UEdGraph>(FocusedDocument))
+		const Workspace::FWorkspaceDocument& FocusedDocument = WorkspaceEditor->GetFocusedWorkspaceDocument();
+		if (UEdGraph* Graph = FocusedDocument.GetTypedObject<UEdGraph>())
 		{
 			// If we just selected a plain graph, rely on default graph search
 			MatchTokensInGraph(Graph, Tokens);
 		}
-		else if (UAnimNextRigVMAsset* AnimNextRigVMAsset = Cast<UAnimNextRigVMAsset>(FocusedDocument))
+		else if (UAnimNextRigVMAsset* AnimNextRigVMAsset = FocusedDocument.GetTypedObject<UAnimNextRigVMAsset>())
 		{
 			// Note: This is where we search subgraphs in say a module
 			if (UAnimNextRigVMAssetEditorData* AnimNextRigVMAssetEditorData = UE::AnimNext::UncookedOnly::FUtils::GetEditorData<UAnimNextRigVMAssetEditorData>(AnimNextRigVMAsset))
diff --git a/Engine/Plugins/Experimental/AnimNext/Source/AnimNextEditor/Private/Common/GraphEditorSchemaActions.cpp b/Engine/Plugins/Experimental/AnimNext/Source/AnimNextEditor/Private/Common/GraphEditorSchemaActions.cpp
--- a/Engine/Plugins/Experimental/AnimNext/Source/AnimNextEditor/Private/Common/GraphEditorSchemaActions.cpp
+++ b/Engine/Plugins/Experimental/AnimNext/Source/AnimNextEditor/Private/Common/GraphEditorSchemaActions.cpp
@@ -4,19 +4,20 @@
 
 #include "AnimNextControllerBase.h"
 #include "EditorUtils.h"
-#include "AnimNextEdGraph.h"
 #include "AnimNextEdGraphNode.h"
 #include "EdGraphNode_Comment.h"
 #include "Editor.h"
-#include "Module/AnimNextModule_EditorData.h"
 #include "RigVMModel/Nodes/RigVMUnitNode.h"
 #include "EdGraphSchema_K2.h"
 #include "GraphEditor.h"
 #include "Framework/MultiBox/MultiBoxBuilder.h"
 #include "Framework/Application/SlateApplication.h"
-#include "AnimNextUnitNode.h"
 #include "AnimNextSharedVariableNode.h"
 #include "Editor/RigVMEditorStyle.h"
+#include "Entries/AnimNextVariableEntry.h"
+#include "AnimNextRigVMAssetEditorData.h"
+#include "AnimNextScopedCompileJob.h"
+#include "ScopedTransaction.h"
 #include "Variables/AnimNextSharedVariables.h"
 
 #define LOCTEXT_NAMESPACE "AnimNextSchemaActions"
@@ -220,6 +221,8 @@ UEdGraphNode* FAnimNextSchemaAction_Variable::PerformAction(UEdGraph* ParentGrap
 
 		FString NodeName;
 
+		FScopedTransaction Transaction(LOCTEXT("AddVariableNode", "Add Variable Node"));
+
 		UAnimNextControllerBase* Controller = CastChecked<UAnimNextControllerBase>(Host->GetRigVMClient()->GetController(ParentGraph));
 		Controller->OpenUndoBracket(TEXT("Add Variable"));
 
@@ -227,6 +230,7 @@ UEdGraphNode* FAnimNextSchemaAction_Variable::PerformAction(UEdGraph* ParentGrap
 		if (!InTypeObjectPath.IsEmpty())
 		{
 			CPPTypeObject = RigVMTypeUtils::FindObjectFromCPPTypeObjectPath<UObject>(InTypeObjectPath);
+
 			if (CPPTypeObject == nullptr)
 			{
 				return nullptr;
@@ -237,9 +241,32 @@ UEdGraphNode* FAnimNextSchemaAction_Variable::PerformAction(UEdGraph* ParentGrap
 		if (!InSourceObjectPath.IsEmpty())
 		{
 			SourceObject = RigVMTypeUtils::FindObjectFromCPPTypeObjectPath<UObject>(InSourceObjectPath);
-			if (SourceObject == nullptr)
+		}
+
+		if (SourceObject == nullptr)
+		{
+			return nullptr;
+		}
+
+		if (UAnimNextRigVMAsset* VariableOwningRigVMAsset = Cast<UAnimNextRigVMAsset>(SourceObject))
+		{
+			// Satisfy host dependency
+			if (UAnimNextRigVMAssetEditorData* OwnerEditorData = ParentGraph->GetTypedOuter<UAnimNextRigVMAssetEditorData>())
 			{
-				return nullptr;
+				if (UAnimNextRigVMAssetEditorData* VariableOwnerEditorData = UE::AnimNext::UncookedOnly::FUtils::GetEditorData(VariableOwningRigVMAsset))
+				{
+					if (OwnerEditorData != VariableOwnerEditorData)
+					{
+						if (!OwnerEditorData->FindEntry(VariableOwningRigVMAsset->GetFName()))
+						{
+							UE::AnimNext::UncookedOnly::FScopedCompileJob CompilerResults(LOCTEXT("AddSharedAssetJobName", "Add Shared Asset dependency"), { ParentGraph->GetTypedOuter<UAnimNextRigVMAsset>() });
+							
+							UAnimNextSharedVariables* VariableOwningSharedVariables = UE::AnimNext::UncookedOnly::FUtils::GetAsset<UAnimNextSharedVariables>(VariableOwnerEditorData);
+							OwnerEditorData->AddSharedVariables(VariableOwningSharedVariables, true, false);
+							OwnerEditorData->RecompileVM();
+						}
+					}
+				}
 			}
 		}
 
@@ -252,7 +279,7 @@ UEdGraphNode* FAnimNextSchemaAction_Variable::PerformAction(UEdGraph* ParentGrap
 		}
 		else
 		{
-			ModelNode = Controller->AddSharedVariableNode(InSourceObjectPath, InName, InTypeName, CPPTypeObject, bIsGetter, FString(), FDeprecateSlateVector2D(Location), NodeName, true, true);
+			ModelNode = Controller->AddSharedVariableNode(InSourceObjectPath, InName, InTypeName, InTypeObjectPath, bIsGetter, FString(), FDeprecateSlateVector2D(Location), NodeName, true, true);
 		}
 
 		if (ModelNode)
@@ -286,7 +313,7 @@ UEdGraphNode* FAnimNextSchemaAction_Variable::PerformAction(UEdGraph* ParentGrap
 
 	if (VariableAccessorChoice == EVariableAccessorChoice::Deferred)
 	{
-		FMenuBuilder MenuBuilder(true, NULL);
+		FMenuBuilder MenuBuilder(true, nullptr);
 		const FText SectionText = FText::FromString(FString::Printf(TEXT("Variable %s"), *Name.ToString()));
 
 		MenuBuilder.BeginSection("VariableDropped", SectionText);
@@ -395,12 +422,6 @@ UEdGraphNode* FAnimNextSchemaAction_Function::PerformAction(UEdGraph* ParentGrap
 	URigVMEdGraphNode* NewNode = nullptr;
 	URigVMEdGraph* EdGraph = Cast<URigVMEdGraph>(ParentGraph);
 
-	UEdGraphPin* FromPin = nullptr;
-	if (FromPins.Num() > 0)
-	{
-		FromPin = FromPins[0];
-	}
-
 	if (Host != nullptr && EdGraph != nullptr)
 	{
 		FName Name = UE::AnimNext::Editor::FUtils::ValidateName(Cast<UObject>(Host), ReferencedPublicFunctionHeader.Name.ToString());
@@ -429,4 +450,86 @@ UEdGraphNode* FAnimNextSchemaAction_Function::PerformAction(UEdGraph* ParentGrap
 	return NewNode;
 }
 
-#undef LOCTEXT_NAMESPACE
+FAnimNextSchemaAction_PromoteToVariable::FAnimNextSchemaAction_PromoteToVariable()
+: FAnimNextSchemaAction(FText(), LOCTEXT("PromoteToVariable", "Promote to Variable..."), LOCTEXT("AddVariableTooltip", "Adds a variable of the pin type to the currently open asset."))
+{
+}
+
+UEdGraphNode* FAnimNextSchemaAction_PromoteToVariable::PerformAction(UEdGraph* ParentGraph, TArray<UEdGraphPin*>& FromPins, const FVector2f& Location, bool bSelectNewNode)
+{
+	
+	UEdGraphNode* Result = nullptr;
+	if (ParentGraph)
+	{
+		IRigVMClientHost* Host = ParentGraph->GetImplementingOuter<IRigVMClientHost>();
+		if(Host == nullptr)
+		{
+			return nullptr;
+		}
+		if (UAnimNextRigVMAssetEditorData* EditorData = Cast<UAnimNextRigVMAssetEditorData>(Host))
+		{
+			if (FromPins.Num() == 1)
+			{
+				if (UEdGraphPin* FromPin = FromPins[0])
+				{
+					UAnimNextRigVMAsset* RigVMAsset = UE::AnimNext::UncookedOnly::FUtils::GetAsset(EditorData);
+					const FAnimNextParamType ParamType = UE::AnimNext::UncookedOnly::FUtils::GetParamTypeFromPinType(FromPin->PinType);
+					if (ParamType.IsValid() && EditorData)
+					{
+						FScopedTransaction Transaction(LOCTEXT("PromoteToVariable", "Promote to Variable"));
+						
+						FString BaseName = FromPin->GetDisplayName().ToString();
+						BaseName.ReplaceCharInline(' ', '_');
+
+						const FName BaseFName = FName(*BaseName);
+						const FName NewVariableName = UE::AnimNext::UncookedOnly::FUtils::GetValidVariableName(EditorData, BaseFName);
+						if (UAnimNextVariableEntry* VariableEntry = EditorData->AddVariable(NewVariableName, ParamType, FromPin->GetDefaultAsString()))
+						{
+							UAnimNextControllerBase* Controller = CastChecked<UAnimNextControllerBase>(Host->GetRigVMClient()->GetController(ParentGraph));
+							Controller->OpenUndoBracket(TEXT("Promote to Variable"));
+							
+							const FAnimNextSchemaAction_Variable::EVariableAccessorChoice AccessorType = FromPin->Direction == EEdGraphPinDirection::EGPD_Input ? FAnimNextSchemaAction_Variable::EVariableAccessorChoice::Get : FAnimNextSchemaAction_Variable::EVariableAccessorChoice::Set;						
+							FAnimNextSchemaAction_Variable Action(VariableEntry->GetVariableName(), RigVMAsset, VariableEntry->GetType(), AccessorType);
+
+							if (URigVMEdGraphNode* RigVMEdNode = Cast<URigVMEdGraphNode>(Action.PerformAction(ParentGraph, FromPins, Location)))
+							{
+								if (URigVMVariableNode* SpawnedVariableNode = Cast<URigVMVariableNode>(RigVMEdNode->GetModelNode()))
+								{
+									const URigVMEdGraphNode* RigVMNode = Cast<URigVMEdGraphNode>(FromPin->GetOwningNode());								
+									const URigVMPin* FromRigPin = RigVMNode->FindModelPinFromGraphPin(FromPin);
+
+									const FString FromPinPath = FromPin->Direction == EEdGraphPinDirection::EGPD_Input ? SpawnedVariableNode->GetValuePin()->GetPinPath() : FromRigPin->GetPinPath();
+									const FString ToPinPath = FromPin->Direction == EEdGraphPinDirection::EGPD_Output ? SpawnedVariableNode->GetValuePin()->GetPinPath() : FromRigPin->GetPinPath();
+									
+									RigVMNode->GetController()->AddLink(FromPinPath, ToPinPath, true, true);
+								}
+
+								Result = RigVMEdNode;
+							}
+
+							if (Result != nullptr)
+							{
+								Controller->CloseUndoBracket();
+							}
+							else
+							{
+								Controller->CancelUndoBracket();
+							}
+						}
+					}
+				}
+			}
+
+			
+		}
+	}
+	
+	return Result;
+}
+
+const FSlateBrush* FAnimNextSchemaAction_PromoteToVariable::GetIconBrush() const
+{
+	return FAnimNextSchemaAction::GetIconBrush();
+}
+
+#undef LOCTEXT_NAMESPACE
\ No newline at end of file
diff --git a/Engine/Plugins/Experimental/AnimNext/Source/AnimNextEditor/Private/Common/SActionMenu.cpp b/Engine/Plugins/Experimental/AnimNext/Source/AnimNextEditor/Private/Common/SActionMenu.cpp
--- a/Engine/Plugins/Experimental/AnimNext/Source/AnimNextEditor/Private/Common/SActionMenu.cpp
+++ b/Engine/Plugins/Experimental/AnimNext/Source/AnimNextEditor/Private/Common/SActionMenu.cpp
@@ -6,19 +6,15 @@
 #include "Framework/Application/SlateApplication.h"
 #include "IDocumentation.h"
 #include "SSubobjectEditor.h"
-#include "AnimNextEdGraphSchema.h"
 #include "RigVMHost.h"
 #include "Widgets/Images/SImage.h"
 #include "Widgets/Layout/SBox.h"
 #include "Widgets/SBoxPanel.h"
 #include "SGraphPalette.h"
-#include "RigVMCore/RigVMRegistry.h"
 #include "RigVMModel/RigVMClient.h"
 #include "RigVMModel/RigVMSchema.h"
 #include "UncookedOnlyUtils.h"
-#include "RigVMBlueprintGeneratedClass.h"
 #include "AnimNextRigVMAssetEditorData.h"
-#include "Entries/AnimNextVariableEntry.h"
 
 #define LOCTEXT_NAMESPACE "AnimNextEditor"
 
@@ -44,7 +40,7 @@ SActionMenu::~SActionMenu()
 	OnCloseReasonCallback.ExecuteIfBound(bActionExecuted, false, !DraggedFromPins.IsEmpty());
 }
 
-void SActionMenu::Construct(const FArguments& InArgs, UEdGraph* InGraph)
+void SActionMenu::Construct(const FArguments& InArgs, UEdGraph* InGraph, const FWorkspaceOutlinerItemExport& InExport)
 {
 	check(InGraph);
 
@@ -67,6 +63,8 @@ void SActionMenu::Construct(const FArguments& InArgs, UEdGraph* InGraph)
 	ContextData.EditorData = ContextData.Graph->GetTypedOuter<UAnimNextRigVMAssetEditorData>();
 	check(ContextData.EditorData);
 
+	ContextData.Export = InExport;
+
 	SBorder::Construct(SBorder::FArguments()
 		.BorderImage(FAppStyle::Get().GetBrush("Menu.Background"))
 		.Padding(5)
diff --git a/Engine/Plugins/Experimental/AnimNext/Source/AnimNextEditor/Private/Variables/Outliner/VariablesOutlinerAssetItem.cpp b/Engine/Plugins/Experimental/AnimNext/Source/AnimNextEditor/Private/Variables/Outliner/VariablesOutlinerAssetItem.cpp
--- a/Engine/Plugins/Experimental/AnimNext/Source/AnimNextEditor/Private/Variables/Outliner/VariablesOutlinerAssetItem.cpp
+++ b/Engine/Plugins/Experimental/AnimNext/Source/AnimNextEditor/Private/Variables/Outliner/VariablesOutlinerAssetItem.cpp
@@ -1,4 +1,4 @@
-﻿// Copyright Epic Games, Inc. All Rights Reserved.
+// Copyright Epic Games, Inc. All Rights Reserved.
 
 #include "VariablesOutlinerAssetItem.h"
 
@@ -7,18 +7,20 @@
 #include "Misc/Attribute.h"
 #include "Misc/Optional.h"
 #include "SceneOutlinerStandaloneTypes.h"
-#include "AnimNextEditorModule.h"
+#include "AnimNextEditorStyle.h"
 #include "AnimNextRigVMAsset.h"
 #include "AnimNextRigVMAssetEditorData.h"
 #include "AssetToolsModule.h"
 #include "AssetViewUtils.h"
-#include "EditorUtils.h"
 #include "IAssetTools.h"
 #include "ScopedTransaction.h"
 #include "UncookedOnlyUtils.h"
+#include "VariablesOutlinerDragDrop.h"
+#include "VariablesOutlinerMode.h"
 #include "Styling/SlateIconFinder.h"
 #include "Styling/StarshipCoreStyle.h"
 #include "UObject/Package.h"
+#include "Variables/SVariablesView.h"
 #include "Widgets/SCompoundWidget.h"
 #include "Widgets/DeclarativeSyntaxSupport.h"
 #include "Widgets/Images/SImage.h"
@@ -26,7 +28,6 @@
 #include "Widgets/Text/SInlineEditableTextBlock.h"
 #include "Widgets/Text/STextBlock.h"
 #include "Widgets/Views/STableRow.h"
-#include "Widgets/Views/SListView.h"
 
 #define LOCTEXT_NAMESPACE "VariablesOutlinerTreeItem"
 
@@ -35,16 +36,17 @@ namespace UE::AnimNext::Editor
 
 const FSceneOutlinerTreeItemType FVariablesOutlinerAssetItem::Type;
 
-class SVariablesOutlinerAssetlabel : FSceneOutlinerCommonLabelData, public SCompoundWidget
+class SVariablesOutlinerAssetLabel : FSceneOutlinerCommonLabelData, public SCompoundWidget
 {
-	SLATE_BEGIN_ARGS(SVariablesOutlinerAssetlabel) {}
+	SLATE_BEGIN_ARGS(SVariablesOutlinerAssetLabel) {}
 	SLATE_END_ARGS()
 
 	void Construct(const FArguments& InArgs, FVariablesOutlinerAssetItem& InTreeItem, ISceneOutliner& SceneOutliner, const STableRow<FSceneOutlinerTreeItemPtr>& InRow)
 	{
 		WeakSceneOutliner = StaticCastSharedRef<ISceneOutliner>(SceneOutliner.AsShared());
 		TreeItem = StaticCastSharedRef<FVariablesOutlinerAssetItem>(InTreeItem.AsShared());
-
+		
+		TWeakPtr<SVariablesOutlinerAssetLabel> WeakLabel = StaticCastSharedRef<SVariablesOutlinerAssetLabel>(AsShared());
 		ChildSlot
 		[
 			SNew(SHorizontalBox)
@@ -54,21 +56,21 @@ class SVariablesOutlinerAssetlabel : FSceneOutlinerCommonLabelData, public SComp
 			.Padding(2.0f, 2.0f)
 			[
 				SNew(SImage)
-				.Image(this, &SVariablesOutlinerAssetlabel::GetAssetIcon)
+				.Image(this, &SVariablesOutlinerAssetLabel::GetAssetIcon)
 			]
 			+SHorizontalBox::Slot()
 			.AutoWidth()
 			.VAlign(VAlign_Center)
-			.Padding(0.0f, 2.0f)
+			.Padding(2.0f, 2.0f)
 			[
 				SAssignNew(TextBlock, SInlineEditableTextBlock)
 				.Font(FStyleFonts::Get().NormalBold)
-				.Text(this, &SVariablesOutlinerAssetlabel::GetDisplayText)
+				.Text(this, &SVariablesOutlinerAssetLabel::GetDisplayText)
 				.HighlightText(SceneOutliner.GetFilterHighlightText())
-				.ColorAndOpacity(this, &SVariablesOutlinerAssetlabel::GetForegroundColor)
-				.OnTextCommitted(this, &SVariablesOutlinerAssetlabel::OnTextCommited)
-				.OnVerifyTextChanged(this, &SVariablesOutlinerAssetlabel::OnVerifyTextChanged)
-				.IsReadOnly(this, &SVariablesOutlinerAssetlabel::CanRename)
+				.ColorAndOpacity(this, &SVariablesOutlinerAssetLabel::GetForegroundColor)
+				.OnTextCommitted(this, &SVariablesOutlinerAssetLabel::OnTextCommited)
+				.OnVerifyTextChanged(this, &SVariablesOutlinerAssetLabel::OnVerifyTextChanged)
+				.IsReadOnly(this, &SVariablesOutlinerAssetLabel::CanRename)
 			]
 			+SHorizontalBox::Slot()
 			.VAlign(VAlign_Center)
@@ -76,20 +78,123 @@ class SVariablesOutlinerAssetlabel : FSceneOutlinerCommonLabelData, public SComp
 			.AutoWidth()
 			[
 				SNew(SImage)
-				.Visibility(this, &SVariablesOutlinerAssetlabel::GetDirtyImageVisibility)
-				.ToolTipText(this, &SVariablesOutlinerAssetlabel::GetDirtyTooltipText)
-				.Image(this, &SVariablesOutlinerAssetlabel::GetDirtyImageBrush)
+				.Visibility(this, &SVariablesOutlinerAssetLabel::GetDirtyImageVisibility)
+				.ToolTipText(this, &SVariablesOutlinerAssetLabel::GetDirtyTooltipText)
+				.Image(this, &SVariablesOutlinerAssetLabel::GetDirtyImageBrush)
+			]
+			+SHorizontalBox::Slot()
+			.VAlign(VAlign_Center)
+			.Padding(2.0f, 0.0f)
+			.AutoWidth()
+			[
+				SNew(SButton)
+				.ButtonStyle(FAppStyle::Get(), "SimpleButton")
+				.OnClicked_Lambda([WeakLabel]()
+				{
+					if (TSharedPtr<SVariablesOutlinerAssetLabel> SharedThis = WeakLabel.Pin())
+					{
+						if (TSharedPtr<FVariablesOutlinerAssetItem> TreeItem = SharedThis->TreeItem.Pin())
+						{
+							if (TSharedPtr<ISceneOutliner> Outliner = SharedThis->WeakSceneOutliner.Pin())
+							{
+								FVariablesOutlinerMode::AddNewVariable(TreeItem->SoftAsset.Get());							
+							}
+						}
+					}
+
+					return FReply::Handled();
+				})
+				.Visibility(this, &SVariablesOutlinerAssetLabel::GetAddVariableButtonVisibility)
+				.ContentPadding(FMargin(1, 0))
+				[
+					SNew(SImage)
+					.Image(FAppStyle::Get().GetBrush("Icons.PlusCircle"))
+					.ColorAndOpacity(FSlateColor::UseForeground())
+				]
+			]
+			+SHorizontalBox::Slot()
+			.VAlign(VAlign_Center)
+			.Padding(0.0f, 0.0f)
+			.AutoWidth()
+			[
+				SAssignNew(ImplementedSharedVariablesHBox, SHorizontalBox)
+				.Clipping(EWidgetClipping::OnDemand)
 			]
+			
 			+SHorizontalBox::Slot()
 			.VAlign(VAlign_Center)
 			.Padding(2.0f, 0.0f)
 			.AutoWidth()
 			[
 				SNew(SThrobber)
-				.Visibility(this, &SVariablesOutlinerAssetlabel::GetLoadingIndicatorVisibility)
+				.Visibility(this, &SVariablesOutlinerAssetLabel::GetLoadingIndicatorVisibility)
 				.ToolTipText(LOCTEXT("LoadingTooltip", "Asset is loading..."))
 			]
 		];
+		
+		if (const TSharedPtr<FVariablesOutlinerAssetItem> Item = TreeItem.Pin())
+		{
+			for (const FSoftObjectPath& SharedVariablesPath : Item->SharedVariableSourcePaths)
+			{
+				const FString SharedVariablesAssetName = SharedVariablesPath.GetAssetName();
+				const FLinearColor BackgroundColor = FLinearColor::MakeRandomSeededColor(GetTypeHash(SharedVariablesAssetName)).Desaturate(0.4f);
+
+				TSharedPtr<SBorder> SharedVariablesBorderWidget;
+				ImplementedSharedVariablesHBox->AddSlot()
+				.VAlign(VAlign_Center)
+				.Padding(2.f, 0.f)
+				.AutoWidth()
+				[
+					SAssignNew(SharedVariablesBorderWidget, SBorder)
+					.BorderImage(FAnimNextEditorStyle::Get().GetBrush("VariablesOutliner.SharedVariablesPill.Border"))
+					.BorderBackgroundColor(BackgroundColor)
+					.ToolTipText( FText::FromString(SharedVariablesAssetName))
+					.VAlign(VAlign_Center)
+					.Padding(FMargin(5.0f, 1.0f, 5.0f, 1.0f))
+					[
+						SNew(STextBlock)
+						.TextStyle(FAnimNextEditorStyle::Get(), "VariablesOutliner.SharedVariablesPill.Text")
+						.Text(FText::FromString(SharedVariablesAssetName))
+						.HighlightText(SceneOutliner.GetFilterHighlightText())						
+						.OverflowPolicy(ETextOverflowPolicy::MiddleEllipsis)
+					]
+				];
+
+				SharedVariablesBorderWidget->SetOnMouseMove(FPointerEventHandler::CreateLambda([SharedVariablesPath, WeakOutliner = WeakSceneOutliner](const FGeometry&, const FPointerEvent& Event) -> FReply
+				{
+					const bool bShouldHighlight = Event.GetModifierKeys().IsControlDown();
+					{
+						if (TSharedPtr<ISceneOutliner> SharedOutliner = WeakOutliner.Pin())
+						{
+							if (const FVariablesOutlinerMode* Mode = static_cast<const FVariablesOutlinerMode*>(SharedOutliner->GetMode()))
+							{
+								if (bShouldHighlight)
+								{
+									Mode->SetHighlightedItem(GetTypeHash(SharedVariablesPath));
+								}
+								else
+								{
+									Mode->ClearHighlightedItem(GetTypeHash(SharedVariablesPath));
+								}
+							}
+						}
+					}
+
+					return FReply::Handled();
+				}));
+
+				SharedVariablesBorderWidget->SetOnMouseLeave(FSimpleNoReplyPointerEventHandler::CreateLambda([SharedVariablesPath, WeakOutliner = WeakSceneOutliner](const FPointerEvent&)
+				{
+					if (TSharedPtr<ISceneOutliner> SharedOutliner = WeakOutliner.Pin())
+					{
+						if (const FVariablesOutlinerMode* Mode = static_cast<const FVariablesOutlinerMode*>(SharedOutliner->GetMode()))
+						{
+							Mode->ClearHighlightedItem(GetTypeHash(SharedVariablesPath));					
+						}
+					}
+				}));
+			}
+		}
 	}
 
 	bool CanRename() const
@@ -192,7 +297,7 @@ class SVariablesOutlinerAssetlabel : FSceneOutlinerCommonLabelData, public SComp
 		return FAppStyle::Get().GetBrush("ClassIcon.Object");
 	}
 
-	void OnTextCommited(const FText& InLabel, ETextCommit::Type InCommitInfo)
+	void OnTextCommited(const FText& InLabel, ETextCommit::Type InCommitInfo) const
 	{
 		if(InCommitInfo == ETextCommit::OnEnter)
 		{
@@ -203,7 +308,7 @@ class SVariablesOutlinerAssetlabel : FSceneOutlinerCommonLabelData, public SComp
 		}
 	}
 
-	bool OnVerifyTextChanged(const FText& InLabel, FText& OutErrorMessage)
+	bool OnVerifyTextChanged(const FText& InLabel, FText& OutErrorMessage) const
 	{
 		if (const TSharedPtr<FVariablesOutlinerAssetItem> Item = TreeItem.Pin())
 		{
@@ -217,14 +322,102 @@ class SVariablesOutlinerAssetlabel : FSceneOutlinerCommonLabelData, public SComp
 		const TOptional<FLinearColor> BaseColor = FSceneOutlinerCommonLabelData::GetForegroundColor(*TreeItem.Pin());
 		return BaseColor.IsSet() ? BaseColor.GetValue() : FSlateColor::UseForeground();
 	}
+	
+	EVisibility GetAddVariableButtonVisibility() const
+	{
+		if (const TSharedPtr<FVariablesOutlinerAssetItem> Item = TreeItem.Pin())
+		{
+			if(Item->SoftAsset.Get())
+			{
+				return EVisibility::Visible;
+			}
+		}
+		return EVisibility::Collapsed;
+	}
+
+	virtual FReply OnDrop(const FGeometry& MyGeometry, const FDragDropEvent& DragDropEvent ) override
+	{
+		TSharedPtr<FVariableDragDropOp> GraphDropOp = DragDropEvent.GetOperationAs<FVariableDragDropOp>();
+		if (GraphDropOp.IsValid())
+		{
+			if (UAnimNextRigVMAsset* Asset = TreeItem.Pin()->SoftAsset.Get())
+			{
+				TSharedPtr<FVariablesOutlinerEntryItem> EntryItem = GraphDropOp->WeakItem.Pin();
+				if (EntryItem.IsValid())
+				{
+					if (UAnimNextVariableEntry* VariableEntry = EntryItem->WeakEntry.Get())
+					{
+						const UAnimNextRigVMAsset* VariableOuter = VariableEntry->GetTypedOuter<UAnimNextRigVMAsset>();
+						if (VariableOuter != Asset)
+						{
+							const FText VariableNameText = FText::FromName(VariableEntry->GetVariableName());
+							const FText AssetNameText = FText::FromName(Asset->GetFName());
+							
+							const FText FormattedMessage = FText::Format(LOCTEXT("MoveVariableToAssetFormat", "Moving {0} to {1}"), VariableNameText, AssetNameText);
+							FScopedTransaction Transaction(FormattedMessage);
+							UncookedOnly::FUtils::MoveVariableToAsset(VariableEntry, Asset, true, true);
+						}
+					}
+				}
+			}
+		}
+		
+		return FReply::Unhandled();
+	}
+
+	virtual void OnDragEnter( const FGeometry& MyGeometry, const FDragDropEvent& DragDropEvent ) override
+	{
+		TSharedPtr<FVariableDragDropOp> GraphDropOp = DragDropEvent.GetOperationAs<FVariableDragDropOp>();
+		if (GraphDropOp.IsValid())
+		{			
+			if (const UAnimNextRigVMAsset* Asset = TreeItem.Pin()->SoftAsset.Get())
+			{
+				TSharedPtr<FVariablesOutlinerEntryItem> EntryItem = GraphDropOp->WeakItem.Pin();
+				if (EntryItem.IsValid())
+				{
+					if (const UAnimNextVariableEntry* VariableEntry = EntryItem->WeakEntry.Get())
+					{
+						const FText VariableNameText = FText::FromName(VariableEntry->GetVariableName());
+						const FText AssetNameText = FText::FromName(Asset->GetFName());
+						
+						const UAnimNextRigVMAsset* VariableOuter = VariableEntry->GetTypedOuter<UAnimNextRigVMAsset>();
+						if (VariableOuter != Asset)
+						{
+							// [TODO] should we only allow move variables "up" the asset chain? If so validate that according to the workspace export chain, that VariableOuter is a child of Asset
+							const FText FormattedMessage = FText::Format(LOCTEXT("MoveVariableToAssetFormat", "Move {0} to {1}"), VariableNameText, AssetNameText);
+							GraphDropOp->SetSimpleFeedbackMessage(nullptr, FSlateColor(GetForegroundColor()), FormattedMessage);				
+						}
+						else
+						{
+							const FSlateBrush* ErrorSymbol = FAppStyle::GetBrush(TEXT("Graph.ConnectorFeedback.Error"));
+							const FText FormattedMessage = FText::Format(LOCTEXT("VariableAlreadyPartOfAssetFormat", "{0} is already part of {1}"), VariableNameText, AssetNameText);
+							GraphDropOp->SetSimpleFeedbackMessage(ErrorSymbol, FSlateColor(GetForegroundColor()), FormattedMessage);		
+						}
+					}
+				}
+			}
+		}
+	}
+
+	virtual void OnDragLeave( const FDragDropEvent& DragDropEvent ) override
+	{
+		TSharedPtr<FVariableDragDropOp> GraphDropOp = DragDropEvent.GetOperationAs<FVariableDragDropOp>();
+		if (GraphDropOp.IsValid())
+		{
+			GraphDropOp->HoverTargetChanged();
+		}
+	}
 
 	TWeakPtr<FVariablesOutlinerAssetItem> TreeItem;
 	TSharedPtr<SInlineEditableTextBlock> TextBlock;
+	TSharedPtr<SHorizontalBox> ImplementedSharedVariablesHBox;
 };
 
-FVariablesOutlinerAssetItem::FVariablesOutlinerAssetItem(TSoftObjectPtr<UAnimNextRigVMAsset> InAsset)
+FVariablesOutlinerAssetItem::FVariablesOutlinerAssetItem(const FItemData& InItemData)
 	: ISceneOutlinerTreeItem(FVariablesOutlinerAssetItem::Type)
-	, SoftAsset(InAsset)
+	, SoftAsset(InItemData.Asset)
+	, SortValue(InItemData.SortValue)
+	, SharedVariableSourcePaths(InItemData.ImplementedSharedVariablesPaths)
 {
 }
 
@@ -235,7 +428,7 @@ bool FVariablesOutlinerAssetItem::IsValid() const
 
 FSceneOutlinerTreeItemID FVariablesOutlinerAssetItem::GetID() const
 {
-	return GetTypeHash(SoftAsset);
+	return GetTypeHash(SoftAsset.ToSoftObjectPath());
 }
 
 FString FVariablesOutlinerAssetItem::GetDisplayString() const
@@ -245,7 +438,7 @@ FString FVariablesOutlinerAssetItem::GetDisplayString() const
 
 TSharedRef<SWidget> FVariablesOutlinerAssetItem::GenerateLabelWidget(ISceneOutliner& Outliner, const STableRow<FSceneOutlinerTreeItemPtr>& InRow)
 {
-	TSharedRef<SVariablesOutlinerAssetlabel> LabelWidget = SNew(SVariablesOutlinerAssetlabel, *this, Outliner, InRow);
+	TSharedRef<SVariablesOutlinerAssetLabel> LabelWidget = SNew(SVariablesOutlinerAssetLabel, *this, Outliner, InRow);
 	RenameRequestEvent.BindSP(LabelWidget->TextBlock.Get(), &SInlineEditableTextBlock::EnterEditingMode);
 	return LabelWidget;
 }
@@ -255,7 +448,7 @@ FString FVariablesOutlinerAssetItem::GetPackageName() const
 	return SoftAsset.GetLongPackageName();
 }
 
-void FVariablesOutlinerAssetItem::Rename(const FText& InNewName)
+void FVariablesOutlinerAssetItem::Rename(const FText& InNewName) const
 {
 	UAnimNextRigVMAsset* Asset = SoftAsset.Get();
 	if(Asset == nullptr)
diff --git a/Engine/Plugins/Experimental/AnimNext/Source/AnimNextEditor/Private/Variables/Outliner/VariablesOutlinerAssetItem.h b/Engine/Plugins/Experimental/AnimNext/Source/AnimNextEditor/Private/Variables/Outliner/VariablesOutlinerAssetItem.h
--- a/Engine/Plugins/Experimental/AnimNext/Source/AnimNextEditor/Private/Variables/Outliner/VariablesOutlinerAssetItem.h
+++ b/Engine/Plugins/Experimental/AnimNext/Source/AnimNextEditor/Private/Variables/Outliner/VariablesOutlinerAssetItem.h
@@ -3,6 +3,7 @@
 #pragma once
 
 #include "ISceneOutlinerTreeItem.h"
+#include "UObject/SoftObjectPtr.h"
 
 class UAnimNextRigVMAsset;
 
@@ -14,8 +15,16 @@ class IWorkspaceOutlinerItemDetails;
 struct FVariablesOutlinerAssetItem : ISceneOutlinerTreeItem
 {
 	static const FSceneOutlinerTreeItemType Type;
+	
+	struct FItemData
+	{
+		const TSoftObjectPtr<UAnimNextRigVMAsset>& Asset;
+		const int32 SortValue;
+		const TArrayView<const FSoftObjectPath> ImplementedSharedVariablesPaths;
+	};
 
-	FVariablesOutlinerAssetItem(TSoftObjectPtr<UAnimNextRigVMAsset> InAsset);
+	FVariablesOutlinerAssetItem(const FItemData& InItemData);
+	virtual ~FVariablesOutlinerAssetItem() override = default;
 
 	// Begin ISceneOutlinerTreeItem overrides
 	virtual bool IsValid() const override;
@@ -27,13 +36,17 @@ struct FVariablesOutlinerAssetItem : ISceneOutlinerTreeItem
 	// End ISceneOutlinerTreeItem overrides
 
 	// Renames the item to the specified name
-	void Rename(const FText& InNewName);
+	void Rename(const FText& InNewName) const;
 
 	// Validates the new item name
 	bool ValidateName(const FText& InNewName, FText& OutErrorMessage) const;
 
 	// Soft ptr to the underlying asset, which may not be loaded yet
 	TSoftObjectPtr<UAnimNextRigVMAsset> SoftAsset;
+
+	int32 SortValue;
+
+	TArray<const FSoftObjectPath> SharedVariableSourcePaths;
 };
 
 }
diff --git a/Engine/Plugins/Experimental/AnimNext/Source/AnimNextEditor/Private/Variables/Outliner/VariablesOutlinerColumns.cpp b/Engine/Plugins/Experimental/AnimNext/Source/AnimNextEditor/Private/Variables/Outliner/VariablesOutlinerColumns.cpp
--- a/Engine/Plugins/Experimental/AnimNext/Source/AnimNextEditor/Private/Variables/Outliner/VariablesOutlinerColumns.cpp
+++ b/Engine/Plugins/Experimental/AnimNext/Source/AnimNextEditor/Private/Variables/Outliner/VariablesOutlinerColumns.cpp
@@ -2,7 +2,6 @@
 
 #include "VariablesOutlinerColumns.h"
 
-#include "AnimNextEditorModule.h"
 #include "AnimNextRigVMAssetEditorData.h"
 #include "DetailLayoutBuilder.h"
 #include "EditorUtils.h"
@@ -19,9 +18,10 @@
 #include "Entries/AnimNextRigVMAssetEntry.h"
 #include "Entries/AnimNextVariableEntry.h"
 #include "Widgets/Images/SImage.h"
-#include "UObject/Package.h"
+#include "Variables/AnimNextUniversalObjectLocatorBindingData.h"
+#include "Variables/AnimNextVariableSettings.h"
 #include "Widgets/Input/SButton.h"
-#include "Widgets/Text/SInlineEditableTextBlock.h"
+#include "Widgets/Input/SCheckBox.h"
 
 #define LOCTEXT_NAMESPACE "VariablesOutlinerColumns"
 
@@ -43,6 +43,7 @@ SHeaderRow::FColumn::FArguments FVariablesOutlinerTypeColumn::ConstructHeaderRow
 		.VAlignHeader(VAlign_Center)
 		.HAlignCell(HAlign_Center)
 		.VAlignCell(VAlign_Center)
+		.SortMode(EColumnSortMode::None)
 		[
 			SNew(SBox)
 			.WidthOverride(16.0f)
@@ -70,7 +71,7 @@ const TSharedRef<SWidget> FVariablesOutlinerTypeColumn::ConstructRowWidget(FScen
 		SNew(SBox)
 		.HAlign(HAlign_Left)
 		.VAlign(VAlign_Center)
-		.IsEnabled(!TreeItem->WeakDataInterfaceEntry.IsValid() || !TreeItem->PropertyPath.IsPathToFieldEmpty())
+		.IsEnabled(!TreeItem->WeakSharedVariablesEntry.IsValid() || !TreeItem->PropertyPath.IsPathToFieldEmpty())
 		[
 			SNew(SPinTypeSelector, FGetPinTypeTree::CreateStatic(&Editor::FUtils::GetFilteredVariableTypeTree))
 			.TargetPinType_Lambda([WeakItem = TWeakPtr<ISceneOutlinerTreeItem>(Item)]()
@@ -104,7 +105,10 @@ const TSharedRef<SWidget> FVariablesOutlinerTypeColumn::ConstructRowWidget(FScen
 					if(ParamType.IsValid())
 					{
 						FScopedTransaction Transaction(LOCTEXT("SetTypeTransaction", "Set Variable Type"));
-						Variable->SetType(ParamType);
+
+						UncookedOnly::FUtils::SetVariableType(Cast<UAnimNextVariableEntry>(Variable), ParamType, true, true);
+						UAnimNextVariableSettings* Settings = GetMutableDefault<UAnimNextVariableSettings>();
+						Settings->SetLastVariableType(ParamType);
 					}
 				}
 			})
@@ -160,9 +164,9 @@ class SVariablesOutlinerValue : public SCompoundWidget, public FNotifyHook
 				EditorData->ModifiedDelegate.AddSP(this, &SVariablesOutlinerValue::HandleModified);
 			}
 		}
-		if(UAnimNextSharedVariablesEntry* DataInterfaceEntry = Cast<UAnimNextSharedVariablesEntry>(InTreeItem.WeakDataInterfaceEntry.Get()))
+		if(const UAnimNextSharedVariablesEntry* SharedVariablesEntry = Cast<UAnimNextSharedVariablesEntry>(InTreeItem.WeakSharedVariablesEntry.Get()))
 		{
-			if(UAnimNextRigVMAssetEditorData* EditorData = DataInterfaceEntry->GetTypedOuter<UAnimNextRigVMAssetEditorData>())
+			if(UAnimNextRigVMAssetEditorData* EditorData = SharedVariablesEntry->GetTypedOuter<UAnimNextRigVMAssetEditorData>())
 			{
 				EditorData->ModifiedDelegate.AddSP(this, &SVariablesOutlinerValue::HandleModified);
 			}
@@ -176,6 +180,8 @@ class SVariablesOutlinerValue : public SCompoundWidget, public FNotifyHook
 		BuildValueWidget();
 	}
 
+	bool bPropertySetGuard = false;
+
 	void BuildValueWidget()
 	{
 		TSharedPtr<FVariablesOutlinerEntryItem> TreeItem = WeakTreeItem.Pin();
@@ -194,6 +200,7 @@ class SVariablesOutlinerValue : public SCompoundWidget, public FNotifyHook
 		FInstancedPropertyBag& PropertyBag = VariableEntry->GetMutablePropertyBag();
 		const FName ValueName = IAnimNextRigVMVariableInterface::ValueName;
 
+		TSharedRef<SHorizontalBox> ValueHBox = SNew(SHorizontalBox);
 		const FPropertyBagPropertyDesc* PropertyDesc = PropertyBag.FindPropertyDescByName(ValueName);
 		if (PropertyDesc != nullptr)
 		{
@@ -209,20 +216,115 @@ class SVariablesOutlinerValue : public SCompoundWidget, public FNotifyHook
 				const TSharedPtr<ISinglePropertyView> SingleStructPropertyView = PropertyEditorModule.CreateSingleProperty(MakeShared<FInstancePropertyBagStructureDataProvider>(PropertyBag), ValueName, SinglePropertyArgs);
 				if (SingleStructPropertyView.IsValid())
 				{
-					ValueWidget = SNew(SHorizontalBox)
-						+SHorizontalBox::Slot()
-						.FillWidth(1.0f)
+					ValueHBox->AddSlot()
+					.AutoWidth()
+					[
+						SNew(SBox)
+						.HAlign(HAlign_Left)
+						.Visibility_Lambda([WeakEntry = TWeakObjectPtr<UAnimNextVariableEntry>(VariableEntry)]() -> EVisibility
+						{
+							if (UAnimNextVariableEntry* Entry = WeakEntry.Get())
+							{
+								return Entry->GetBinding().IsValid() ? EVisibility::Collapsed : EVisibility::Visible;
+							}
+							return EVisibility::Collapsed;
+						})
+						[
+							SingleStructPropertyView.ToSharedRef()
+						]
+					];
+
+					const TSharedPtr<ISinglePropertyView> LocatorPropertyView = PropertyEditorModule.CreateSingleProperty(VariableEntry, GET_MEMBER_NAME_CHECKED(UAnimNextVariableEntry, Binding), SinglePropertyArgs);
+
+					LocatorPropertyView.Get()->SetOnPropertyValueChanged(FSimpleDelegate::CreateLambda([this, LocatorPropertyView, WeakEntry = TWeakObjectPtr<UAnimNextVariableEntry>(VariableEntry)]()
+					{
+						if (bPropertySetGuard)
+						{
+							return;
+						}
+
+						TGuardValue<bool> Guard(bPropertySetGuard, true);						
+						if (UAnimNextVariableEntry* Entry = WeakEntry.Get())
+						{
+							Entry->BroadcastModified(EAnimNextEditorDataNotifType::VariableBindingChanged);
+						}
+
+						LocatorPropertyView->GetPropertyHandle()->NotifyPostChange(EPropertyChangeType::ValueSet);
+					}));
+					
+					ValueHBox->AddSlot()
+					.AutoWidth()
+					[
+						SNew(SBox)
+						.HAlign(HAlign_Left)
+						.Visibility_Lambda([WeakEntry = TWeakObjectPtr<UAnimNextVariableEntry>(VariableEntry)]() -> EVisibility
+						{
+							if (UAnimNextVariableEntry* Entry = WeakEntry.Get())
+							{
+								return Entry->GetBinding().IsValid() ? EVisibility::Visible : EVisibility::Collapsed;
+							}
+
+							return EVisibility::Collapsed;
+						})
+						[
+							VariableEntry ? LocatorPropertyView.ToSharedRef() : SNullWidget::NullWidget
+						]
+					];
+
+					ValueHBox->AddSlot()
+						.AutoWidth()
+						.VAlign(VAlign_Center)
+						.HAlign(HAlign_Center)
 						[
-							SNew(SBox)
-							.HAlign(HAlign_Left)
+							SNew(SCheckBox)
+							.Style( &FAppStyle::Get().GetWidgetStyle<FCheckBoxStyle>("ToggleButtonCheckBox"))
+							.Padding(FMargin(4.f, 2.f))
+							.IsChecked_Lambda([WeakEntry = TWeakObjectPtr<UAnimNextVariableEntry>(VariableEntry)]() -> ECheckBoxState
+							{
+								if (UAnimNextVariableEntry* Entry = WeakEntry.Get())
+								{
+									return Entry->GetBinding().IsValid() ? ECheckBoxState::Checked : ECheckBoxState::Unchecked;
+								}
+
+								return ECheckBoxState::Undetermined;
+							})
 							[
-								SingleStructPropertyView.ToSharedRef()
+								SNew(SImage)
+								.Image_Lambda([WeakEntry = TWeakObjectPtr<UAnimNextVariableEntry>(VariableEntry)]() -> const FSlateBrush*
+								{
+									if (UAnimNextVariableEntry* Entry = WeakEntry.Get())
+									{
+										return Entry->GetBinding().IsValid() ? FAppStyle::Get().GetBrush("Icons.Link") : FAppStyle::Get().GetBrush("Icons.Unlink");
+									}
+
+									return FAppStyle::GetBrush(TEXT("ClassIcon.Default"));
+								})
+								.ColorAndOpacity(FSlateColor::UseForeground())
 							]
-						];
+							.OnCheckStateChanged_Lambda([LocatorPropertyView, WeakEntry = TWeakObjectPtr<UAnimNextVariableEntry>(VariableEntry)](ECheckBoxState InState)
+							{
+								if (UAnimNextVariableEntry* Entry = WeakEntry.Get())
+								{
+									if (InState == ECheckBoxState::Checked)
+									{
+										Entry->SetBindingType(FAnimNextUniversalObjectLocatorBindingData::StaticStruct(), true);
+										LocatorPropertyView->GetPropertyHandle()->NotifyPostChange(EPropertyChangeType::ValueSet);
+									}
+									else if (InState == ECheckBoxState::Unchecked)
+									{
+										Entry->SetBindingType(nullptr, true);
+										LocatorPropertyView->GetPropertyHandle()->NotifyPostChange(EPropertyChangeType::ValueSet);
+									}
+
+									Entry->BroadcastModified(EAnimNextEditorDataNotifType::VariableBindingChanged);
+								}
+							})
+					];
 				}
 			}
 		}
 
+		ValueWidget = ValueHBox;
 		WidgetContainer->SetContent(ValueWidget.ToSharedRef());
 	}
 
@@ -254,6 +356,7 @@ class SVariablesOutlinerValue : public SCompoundWidget, public FNotifyHook
 	void HandleModified(UAnimNextRigVMAssetEditorData* InEditorData, EAnimNextEditorDataNotifType InType, UObject* InSubject)
 	{
 		if( InType != EAnimNextEditorDataNotifType::VariableDefaultValueChanged &&
+			InType != EAnimNextEditorDataNotifType::VariableTypeChanged &&
 			InType != EAnimNextEditorDataNotifType::UndoRedo)
 		{
 			return;
@@ -266,8 +369,8 @@ class SVariablesOutlinerValue : public SCompoundWidget, public FNotifyHook
 		}
 
 		UAnimNextVariableEntry* VariableEntry = Cast<UAnimNextVariableEntry>(TreeItem->WeakEntry.Get());
-		UAnimNextSharedVariablesEntry* DataInterfaceEntry = Cast<UAnimNextSharedVariablesEntry>(TreeItem->WeakDataInterfaceEntry.Get());
-		if(VariableEntry != InSubject && DataInterfaceEntry != InSubject)
+		const UAnimNextSharedVariablesEntry* SharedVariablesEntry = Cast<UAnimNextSharedVariablesEntry>(TreeItem->WeakSharedVariablesEntry.Get());
+		if(VariableEntry != InSubject && SharedVariablesEntry != InSubject)
 		{
 			return;
 		}
@@ -338,11 +441,12 @@ class SVariablesOutlinerAccessSpecifier : public SCompoundWidget, public FNotify
 		ChildSlot
 		[
 			SNew(SBox)
-			.IsEnabled(!InTreeItem.WeakDataInterfaceEntry.IsValid())
+			.IsEnabled(!InTreeItem.WeakSharedVariablesEntry.IsValid())
 			.WidthOverride(16.0f)
 			.HeightOverride(16.0f)
 			.VAlign(VAlign_Center)
 			.HAlign(HAlign_Center)
+			.IsEnabled(this, &SVariablesOutlinerAccessSpecifier::IsButtonEnabled)
 			[
 				SNew(SButton)
 				.ButtonStyle(FAppStyle::Get(), "HoverHintOnly")
@@ -358,7 +462,7 @@ class SVariablesOutlinerAccessSpecifier : public SCompoundWidget, public FNotify
 		];
 	}
 
-	FReply OnClicked()
+	FReply OnClicked() const
 	{
 		TSharedPtr<FVariablesOutlinerEntryItem> TreeItem = WeakTreeItem.Pin();
 		if(!TreeItem.IsValid())
@@ -397,6 +501,26 @@ class SVariablesOutlinerAccessSpecifier : public SCompoundWidget, public FNotify
 			FAppStyle::GetBrush("Level.NotVisibleHighlightIcon16x");
 	}
 
+	bool IsButtonEnabled() const
+	{
+		TSharedPtr<FVariablesOutlinerEntryItem> TreeItem = WeakTreeItem.Pin();
+		if(!TreeItem.IsValid())
+		{
+			return false;
+		}
+
+		if (UAnimNextVariableEntry* Entry = TreeItem->WeakEntry.Get())
+		{
+			if (UAnimNextRigVMAsset* OuterAsset = Entry->GetTypedOuter<UAnimNextRigVMAsset>())
+			{
+				// Disable toggling public/private for pure UAnimNextSharedVariables objects				
+				return ExactCast<UAnimNextSharedVariables>(OuterAsset) == nullptr;	
+			}
+		}
+		
+		return false;		
+	}	
+
 	FText GetTooltipText() const
 	{
 		TSharedPtr<FVariablesOutlinerEntryItem> TreeItem = WeakTreeItem.Pin();
diff --git a/Engine/Plugins/Experimental/AnimNext/Source/AnimNextEditor/Private/Variables/Outliner/VariablesOutlinerCommands.cpp b/Engine/Plugins/Experimental/AnimNext/Source/AnimNextEditor/Private/Variables/Outliner/VariablesOutlinerCommands.cpp
new file mode 100644
--- /dev/null
+++ b/Engine/Plugins/Experimental/AnimNext/Source/AnimNextEditor/Private/Variables/Outliner/VariablesOutlinerCommands.cpp
@@ -0,0 +1,22 @@
+// Copyright Epic Games, Inc. All Rights Reserved.
+
+
+#include "VariablesOutlinerCommands.h"
+
+#define LOCTEXT_NAMESPACE "VariablesOutlinerCommands"
+
+void FVariablesOutlinerCommands::RegisterCommands()
+{
+	UI_COMMAND(AddNewVariable, "Add Variable", "Adds a single new variable to the selected asset.", EUserInterfaceActionType::Button, FInputChord());
+	UI_COMMAND(AddNewVariables, "Add Variable(s)", "Adds variables to assets.\nIf multiple assets are selected, then variables will be added to each.\nIf no assets are selected and there are multiple assets, variables will be added to all assets.", EUserInterfaceActionType::Button, FInputChord());
+
+	UI_COMMAND(FindReferences, "In Project", "Finds all references to the selected variable variables across assets in the project.", EUserInterfaceActionType::Button, FInputChord());
+	UI_COMMAND(FindReferencesInWorkspace, "In Workspace", "Finds all references to the selected variable variables across assets in the current workspace.", EUserInterfaceActionType::Button, FInputChord());
+	UI_COMMAND(FindReferencesInAsset, "In Asset", "Finds all references to the selected variable variables in the currently opened asset.", EUserInterfaceActionType::Button, FInputChord());
+
+	UI_COMMAND(SaveAsset, "Save Asset", "Save the selected asset.", EUserInterfaceActionType::Button, FInputChord());
+
+	UI_COMMAND(ToggleVariableExport, "Toggle Variable Export", "Sets Variable to be either Public or Private.", EUserInterfaceActionType::Button, FInputChord());
+}
+
+#undef LOCTEXT_NAMESPACE
\ No newline at end of file
diff --git a/Engine/Plugins/Experimental/AnimNext/Source/AnimNextEditor/Private/Variables/Outliner/VariablesOutlinerEntryItem.h b/Engine/Plugins/Experimental/AnimNext/Source/AnimNextEditor/Private/Variables/Outliner/VariablesOutlinerEntryItem.h
--- a/Engine/Plugins/Experimental/AnimNext/Source/AnimNextEditor/Private/Variables/Outliner/VariablesOutlinerEntryItem.h
+++ b/Engine/Plugins/Experimental/AnimNext/Source/AnimNextEditor/Private/Variables/Outliner/VariablesOutlinerEntryItem.h
@@ -27,19 +27,25 @@ struct FVariablesOutlinerEntryItem : ISceneOutlinerTreeItem
 	// End ISceneOutlinerTreeItem overrides
 
 	// Renames the item to the specified name
-	void Rename(const FText& InNewName);
+	void Rename(const FText& InNewName) const;
 
 	// Validates the new item name
 	bool ValidateName(const FText& InNewName, FText& OutErrorMessage) const;
 
+	bool HasStructOwner() const { return bStructOwner; }
+
 	// Ptr to the underlying entry if this variable came from an asset
 	TWeakObjectPtr<UAnimNextVariableEntry> WeakEntry;
 
 	// The data interface entry this entry is from, if any
-	TWeakObjectPtr<UAnimNextSharedVariablesEntry> WeakDataInterfaceEntry;
+	TWeakObjectPtr<const UAnimNextSharedVariablesEntry> WeakSharedVariablesEntry;
 
 	// Ptr to the property if this variable came from a struct
 	TFieldPath<const FProperty> PropertyPath;
+
+	bool bStructOwner = false;
+
+	int32 SortValue = 0;
 };
 
 }
