AnimNext variables UI/UX changes:

diff --git a/Engine/Plugins/Experimental/AnimNext/Source/AnimNextEditor/Private/Variables/Outliner/VariablesOutlinerHierarchy.cpp b/Engine/Plugins/Experimental/AnimNext/Source/AnimNextEditor/Private/Variables/Outliner/VariablesOutlinerHierarchy.cpp
--- a/Engine/Plugins/Experimental/AnimNext/Source/AnimNextEditor/Private/Variables/Outliner/VariablesOutlinerHierarchy.cpp
+++ b/Engine/Plugins/Experimental/AnimNext/Source/AnimNextEditor/Private/Variables/Outliner/VariablesOutlinerHierarchy.cpp
@@ -6,9 +6,9 @@
 #include "ISceneOutlinerMode.h"
 #include "AnimNextRigVMAssetEditorData.h"
 #include "VariablesOutlinerAssetItem.h"
-#include "VariablesOutlinerDataInterfaceItem.h"
+#include "VariablesOutlinerCategoryItem.h"
+#include "VariablesOutlinerStructSharedVariablesItem.h"
 #include "VariablesOutlinerMode.h"
-#include "WorkspaceAssetRegistryInfo.h"
 #include "VariablesOutlinerEntryItem.h"
 #include "Entries/AnimNextSharedVariablesEntry.h"
 #include "Entries/AnimNextRigVMAssetEntry.h"
@@ -19,6 +19,8 @@
 namespace UE::AnimNext::Editor
 {
 
+TAutoConsoleVariable<bool> GOutputHierarchyDebugInformationCvar(TEXT("VariablesOutliner.OutputHierarchyDebugInformation"), false, TEXT("If enabled will output debug information about the hierarchy generation to the output log")); 
+
 FVariablesOutlinerHierarchy::FVariablesOutlinerHierarchy(ISceneOutlinerMode* Mode)
 	: ISceneOutlinerHierarchy(Mode)
 {
@@ -27,109 +29,32 @@ FVariablesOutlinerHierarchy::FVariablesOutlinerHierarchy(ISceneOutlinerMode* Mod
 void FVariablesOutlinerHierarchy::CreateItems(TArray<FSceneOutlinerTreeItemPtr>& OutItems) const
 {
 	SVariablesOutliner* Outliner = static_cast<FVariablesOutlinerMode* const>(Mode)->GetOutliner();
-	for(const TSoftObjectPtr<UAnimNextRigVMAsset>& SoftAsset : Outliner->Assets)
-	{
-		UAnimNextRigVMAsset* Asset = SoftAsset.Get();
-		if(Asset == nullptr)
-		{
-			continue;
-		}
-
-		const UAnimNextRigVMAssetEditorData* EditorData = UncookedOnly::FUtils::GetEditorData<UAnimNextRigVMAssetEditorData>(Asset);
-		if (EditorData == nullptr)
-		{
-			continue;
-		}
 
-		if (FSceneOutlinerTreeItemPtr Item = Mode->CreateItemFor<FVariablesOutlinerAssetItem>(SoftAsset))
+	if(Outliner->Assets.Num())
+	{
+		// Process first asset only for now
+		const int32 NumAssets = Outliner->Assets.Num();
+		for(int32 Index = 0; Index < NumAssets; Index++)
 		{
-			OutItems.Add(Item);
+			const TSoftObjectPtr<UAnimNextRigVMAsset>& SoftAsset = Outliner->Assets[Index];
+			TArray<const FSoftObjectPath> SharedVariableSourcePaths;
+			ProcessAsset(SoftAsset, (NumAssets - Index), OutItems, SharedVariableSourcePaths);
 		}
-
-		EditorData->ForEachEntryOfType<UAnimNextVariableEntry>([this, &OutItems](UAnimNextVariableEntry* InVariable)
-		{
-			if (FSceneOutlinerTreeItemPtr Item = Mode->CreateItemFor<FVariablesOutlinerEntryItem>(CastChecked<UAnimNextVariableEntry>(InVariable)))
-			{
-				OutItems.Add(Item);
-			}
-			return true;
-		});
-
-		EditorData->ForEachEntryOfType<UAnimNextSharedVariablesEntry>([this, &OutItems](UAnimNextSharedVariablesEntry* InInterfaceEntry)
-		{
-			if (FSceneOutlinerTreeItemPtr Item = Mode->CreateItemFor<FVariablesOutlinerDataInterfaceItem>(CastChecked<UAnimNextSharedVariablesEntry>(InInterfaceEntry)))
-			{
-				OutItems.Add(Item);
-			}
-
-			auto AddDataInterface = [this, &OutItems](UAnimNextSharedVariablesEntry* InDataInterfaceEntry, UAnimNextSharedVariablesEntry* InRootDataInterfaceEntry, auto& InAddDataInterface) -> void
-			{
-				switch (InDataInterfaceEntry->GetType())
-				{
-				case EAnimNextSharedVariablesType::Asset:
-					{
-						UAnimNextRigVMAssetEditorData* InterfaceEditorData = UncookedOnly::FUtils::GetEditorData<UAnimNextRigVMAssetEditorData>(InDataInterfaceEntry->GetAsset());
-						InterfaceEditorData->ForEachEntryOfType<UAnimNextVariableEntry>([this, &OutItems, InRootDataInterfaceEntry](UAnimNextVariableEntry* InVariable)
-						{
-							if(InVariable->GetExportAccessSpecifier() == EAnimNextExportAccessSpecifier::Public)
-							{
-								if (FSceneOutlinerTreeItemPtr Item = Mode->CreateItemFor<FVariablesOutlinerEntryItem>(CastChecked<UAnimNextVariableEntry>(InVariable)))
-								{
-									FVariablesOutlinerEntryItem* EntryItem = Item->CastTo<FVariablesOutlinerEntryItem>();
-									EntryItem->WeakDataInterfaceEntry = InRootDataInterfaceEntry;
-									OutItems.Add(Item);
-								}
-							}
-							return true;
-						});
-
-						InterfaceEditorData->ForEachEntryOfType<UAnimNextSharedVariablesEntry>([this, &OutItems, &InAddDataInterface, InRootDataInterfaceEntry](UAnimNextSharedVariablesEntry* InSubInterfaceEntry)
-						{
-							InAddDataInterface(InSubInterfaceEntry, InRootDataInterfaceEntry, InAddDataInterface);
-							return true;
-						});
-						break;
-					}
-				case EAnimNextSharedVariablesType::Struct:
-					{
-						const UScriptStruct* Struct = InDataInterfaceEntry->GetStruct();
-						for (TFieldIterator<FProperty> It(Struct); It; ++It)
-						{
-							const FProperty* Property = *It;
-							if (FSceneOutlinerTreeItemPtr Item = Mode->CreateItemFor<FVariablesOutlinerEntryItem>(Property))
-							{
-								FVariablesOutlinerEntryItem* EntryItem = Item->CastTo<FVariablesOutlinerEntryItem>();
-								EntryItem->WeakDataInterfaceEntry = InRootDataInterfaceEntry;
-								OutItems.Add(Item);
-							}
-						}
-						break;
-					}
-				}
-			};
-
-			AddDataInterface(InInterfaceEntry, InInterfaceEntry, AddDataInterface);
-
-			return true;
-		});
 	}
 }
 
 FSceneOutlinerTreeItemPtr FVariablesOutlinerHierarchy::FindOrCreateParentItem(const ISceneOutlinerTreeItem& Item, const TMap<FSceneOutlinerTreeItemID, FSceneOutlinerTreeItemPtr>& Items, bool bCreate)
 {
+	uint32 ParentHash = INDEX_NONE;
+	
 	if (const FVariablesOutlinerEntryItem* EntryItem = Item.CastTo<FVariablesOutlinerEntryItem>())
 	{
-		UAnimNextSharedVariablesEntry* SharedVariablesEntry = EntryItem->WeakDataInterfaceEntry.Get();
+		const UAnimNextSharedVariablesEntry* SharedVariablesEntry = EntryItem->WeakSharedVariablesEntry.Get();
+		
+		FSoftObjectPath OwnerPath;
 		if(SharedVariablesEntry != nullptr)
 		{
-			// Add as part of an imported data interface, so use that as parent
-			TSoftObjectPtr<UAnimNextRigVMAssetEntry> SoftObjectPtr(SharedVariablesEntry); 
-			uint32 Hash = GetTypeHash(SoftObjectPtr);
-
-			if (const FSceneOutlinerTreeItemPtr* ParentItem = Items.Find(Hash))
-			{
-				return *ParentItem;
-			}
+			OwnerPath = SharedVariablesEntry->GetObjectPath();
 		}
 		else
 		{
@@ -146,39 +71,159 @@ FSceneOutlinerTreeItemPtr FVariablesOutlinerHierarchy::FindOrCreateParentItem(co
 			}
 
 			TSoftObjectPtr<UAnimNextRigVMAsset> SoftObjectPtr(Asset);
-			uint32 Hash = GetTypeHash(SoftObjectPtr);
+			OwnerPath = Asset;
+			
+		}
+
+		uint32 CategoryHash = INDEX_NONE;
 
-			if (const FSceneOutlinerTreeItemPtr* ParentItem = Items.Find(Hash))
+		// Now we have the parent asset hash, determine if we need to be child-ed to a category
+		if (UAnimNextVariableEntry* VariableEntry = EntryItem->WeakEntry.Get())
+		{
+			if (!VariableEntry->GetVariableCategory().IsEmpty())
 			{
-				return *ParentItem;
+				CategoryHash = GetTypeHash(VariableEntry->GetVariableCategory());
 			}
 		}
+
+		check(OwnerPath.IsValid());
+		ParentHash = CategoryHash != INDEX_NONE ? HashCombine(GetTypeHash(OwnerPath), CategoryHash) : GetTypeHash(OwnerPath);		
+		
 	}
 
-	if (const FVariablesOutlinerDataInterfaceItem* InterfaceItem = Item.CastTo<FVariablesOutlinerDataInterfaceItem>())
+	if (const FVariablesOutlinerCategoryItem* CategoryItem = Item.CastTo<FVariablesOutlinerCategoryItem>())
 	{
-		UAnimNextRigVMAssetEntry* Entry = InterfaceItem->WeakEntry.Get();
-		if (Entry == nullptr)
+		if (UAnimNextRigVMAsset* Asset = CategoryItem->WeakOwner.Get())
 		{
-			return nullptr;
+			const FSoftObjectPath OwnerPath = Asset;
+			ParentHash = GetTypeHash(OwnerPath);
 		}
+	}
 
-		UAnimNextRigVMAsset* Asset = Entry->GetTypedOuter<UAnimNextRigVMAsset>();
-		if (Entry == nullptr)
+	if (const FSceneOutlinerTreeItemPtr* ParentItem = Items.Find(ParentHash))
+	{
+		return *ParentItem;
+	}
+
+	return nullptr;
+}
+
+void FVariablesOutlinerHierarchy::ProcessAsset(const TSoftObjectPtr<UAnimNextRigVMAsset>& InSoftAsset, int32 Depth, TArray<FSceneOutlinerTreeItemPtr>& OutItems, TArray<const FSoftObjectPath>& OutSharedVariableSourcePaths) const
+{
+	UAnimNextRigVMAsset* Asset = InSoftAsset.Get();
+	if(Asset == nullptr)
+	{
+		return;
+	}
+		
+	const UAnimNextRigVMAssetEditorData* EditorData = UncookedOnly::FUtils::GetEditorData<UAnimNextRigVMAssetEditorData>(Asset);
+	if (EditorData == nullptr)
+	{
+		return;
+	}
+
+	UE_CLOG(GOutputHierarchyDebugInformationCvar.GetValueOnAnyThread(),LogAnimation, Warning, TEXT("ProcessAsset: %s %i"), *InSoftAsset.GetAssetName(), Depth);
+		
+	TArray<FSceneOutlinerTreeItemPtr> ParentItems;
+	EditorData->ForEachEntryOfType<UAnimNextSharedVariablesEntry>([this, &OutSharedVariableSourcePaths, &ParentItems, Depth, &OutItems](UAnimNextSharedVariablesEntry* InSharedVariablesEntry)
+	{
+		SVariablesOutliner* Outliner = static_cast<FVariablesOutlinerMode* const>(Mode)->GetOutliner();
+
+		switch (InSharedVariablesEntry->Type)
 		{
-			return nullptr;
-		}
+		case EAnimNextSharedVariablesType::Asset:
+			{
+				const UAnimNextSharedVariables* SharedVariables = InSharedVariablesEntry->GetAsset();
+				OutSharedVariableSourcePaths.Add(SharedVariables);
+				
+				// Would have expected shared variable (assets) to be part of the outer export chain, so should have been handled as regular assets.
+				if (!Outliner->Assets.Contains(SharedVariables))
+				{	
+					TArray<const FSoftObjectPath> SharedVariableSourcePaths;
+					ProcessAsset( const_cast<UAnimNextSharedVariables*>(SharedVariables), Depth - 1, ParentItems, SharedVariableSourcePaths);
+				}
+			}
+			break;
+
+		case EAnimNextSharedVariablesType::Struct:
+			{
+				if (FSceneOutlinerTreeItemPtr Item = Mode->CreateItemFor<FVariablesOutlinerStructSharedVariablesItem>(InSharedVariablesEntry))
+				{
+					OutItems.Add(Item);
+				}
+
+				UE_CLOG(GOutputHierarchyDebugInformationCvar.GetValueOnAnyThread(), LogAnimation, Warning, TEXT("~ FVariablesOutlinerStructSharedVariablesItem: %s %i"), *InSharedVariablesEntry->GetObjectPath().ToString(), Depth);
+				
+				const UScriptStruct* Struct = InSharedVariablesEntry->GetStruct();
+				OutSharedVariableSourcePaths.Add(Struct);
+				int32 Index = 0;
+				for (TFieldIterator<FProperty> It(Struct); It; ++It)
+				{
+					const FProperty* Property = *It;
+					if (FSceneOutlinerTreeItemPtr Item = Mode->CreateItemFor<FVariablesOutlinerEntryItem>(Property))
+					{
+						FVariablesOutlinerEntryItem* EntryItem = Item->CastTo<FVariablesOutlinerEntryItem>();
+						EntryItem->WeakSharedVariablesEntry = InSharedVariablesEntry;
+						EntryItem->SortValue = Index++;
+						OutItems.Add(Item);
+					}
+				}
+			}
+			break;
+		}		
+		
+		return true;
+	});
 
-		TSoftObjectPtr<UAnimNextRigVMAsset> SoftObjectPtr(Asset);
-		uint32 Hash = GetTypeHash(SoftObjectPtr);
+	OutItems.Append(ParentItems);
 
-		if (const FSceneOutlinerTreeItemPtr* ParentItem = Items.Find(Hash))
+	TMap<FString, FSceneOutlinerTreeItemPtr> CategoryItems;
+
+	for (int32 CategoryIndex = 0; CategoryIndex < EditorData->VariableAndFunctionCategories.Num(); ++CategoryIndex)
+	{
+		const FString& CategoryName = EditorData->VariableAndFunctionCategories[CategoryIndex];
+		if (FSceneOutlinerTreeItemPtr Item = Mode->CreateItemFor<FVariablesOutlinerCategoryItem>(FVariablesOutlinerCategoryItem::FItemData{Asset, CategoryName}))
 		{
-			return *ParentItem;
+			Item->CastTo<FVariablesOutlinerCategoryItem>()->SortValue = CategoryIndex;
+			CategoryItems.Add(CategoryName, Item);
+
+			UE_CLOG(GOutputHierarchyDebugInformationCvar.GetValueOnAnyThread(),LogAnimation, Warning, TEXT("+ FVariablesOutlinerCategoryItem: %s %i"), *CategoryName, CategoryIndex);
 		}
 	}
 
-	return nullptr;
-}
+	// Adjust asset depth to ensure SharedVariables are always pinned to 0 depth
+	const int32 AssetDepth = Asset->GetClass() && Asset->GetClass() == UAnimNextSharedVariables::StaticClass() ? 0 : Depth;
+
+	constexpr int32 NumberOfVariableSortValueEntriesInCategory = 1024;
+
+	if (FSceneOutlinerTreeItemPtr Item = Mode->CreateItemFor<FVariablesOutlinerAssetItem>(FVariablesOutlinerAssetItem::FItemData{InSoftAsset, AssetDepth, OutSharedVariableSourcePaths}))
+	{
+		OutItems.Add(Item);
+
+		UE_CLOG(GOutputHierarchyDebugInformationCvar.GetValueOnAnyThread(), LogAnimation, Warning, TEXT("* FVariablesOutlinerAssetItem: %s %i"), *InSoftAsset.GetAssetName(), AssetDepth);
+
+		int32 Index = 1;
+		EditorData->ForEachEntryOfType<UAnimNextVariableEntry>([this, &OutItems, &CategoryItems, &Index, EditorData](UAnimNextVariableEntry* InVariable)
+		{
+			if (FSceneOutlinerTreeItemPtr Item = Mode->CreateItemFor<FVariablesOutlinerEntryItem>(CastChecked<UAnimNextVariableEntry>(InVariable)))
+			{
+				FVariablesOutlinerEntryItem* EntryItem = Item->CastTo<FVariablesOutlinerEntryItem>();
+				const int32 CategoryIndex = EditorData->VariableAndFunctionCategories.IndexOfByKey(InVariable->Category);
+				EntryItem->SortValue = (CategoryIndex != INDEX_NONE ? ((CategoryIndex + 1) * NumberOfVariableSortValueEntriesInCategory) : 0) + Index++;
+				OutItems.Add(Item);
+				
+				OutItems.Add(Item);
+				UE_CLOG(GOutputHierarchyDebugInformationCvar.GetValueOnAnyThread(), LogAnimation, Warning, TEXT("\t- FVariablesOutlinerEntryItem: %s %i"), *InVariable->GetEntryName().ToString(), Item->CastTo<FVariablesOutlinerEntryItem>()->SortValue);
 
+				// Add category item on demand (to ensure we do not end up with empty category item hierarchies)
+				if (!InVariable->Category.IsEmpty())
+				{
+					FSceneOutlinerTreeItemPtr& CategoryItemToAdd = CategoryItems.FindChecked(InVariable->Category);
+					OutItems.AddUnique(CategoryItemToAdd);
+				}
+			}
+			return true;
+		});
+	}
+}
 }
diff --git a/Engine/Plugins/Experimental/AnimNext/Source/AnimNextEditor/Private/Variables/Outliner/VariablesOutlinerMode.cpp b/Engine/Plugins/Experimental/AnimNext/Source/AnimNextEditor/Private/Variables/Outliner/VariablesOutlinerMode.cpp
--- a/Engine/Plugins/Experimental/AnimNext/Source/AnimNextEditor/Private/Variables/Outliner/VariablesOutlinerMode.cpp
+++ b/Engine/Plugins/Experimental/AnimNext/Source/AnimNextEditor/Private/Variables/Outliner/VariablesOutlinerMode.cpp
@@ -10,7 +10,6 @@
 #include "AnimNextRigVMAssetEditorData.h"
 #include "IWorkspaceEditor.h"
 #include "AnimNextEditorModule.h"
-#include "AnimNextRigVMAssetEditorData.h"
 #include "AssetRegistry/AssetRegistryModule.h"
 #include "ScopedTransaction.h"
 #include "ToolMenus.h"
@@ -19,14 +18,20 @@
 #include "Variables/SAddVariablesDialog.h"
 #include "Variables/AnimNextVariableItemMenuContext.h"
 #include "AnimNextRigVMAsset.h"
+#include "FileHelpers.h"
 #include "VariablesOutlinerAssetItem.h"
-#include "VariablesOutlinerDataInterfaceItem.h"
+#include "VariablesOutlinerCommands.h"
+#include "VariablesOutlinerStructSharedVariablesItem.h"
+#include "VariablesOutlinerCategoryItem.h"
 #include "VariablesOutlinerDragDrop.h"
+#include "Common/AnimNextAssetFindReplaceVariables.h"
 #include "Common/GraphEditorSchemaActions.h"
 #include "Entries/AnimNextVariableEntry.h"
 #include "Variables/SVariablesView.h"
 #include "Entries/AnimNextSharedVariablesEntry.h"
 #include "Variables/AnimNextSharedVariables.h"
+#include "Variables/AnimNextVariableSettings.h"
+#include "Widgets/Docking/SDockTab.h"
 
 #define LOCTEXT_NAMESPACE "VariablesOutlinerMode"
 
@@ -63,23 +68,52 @@ TSharedPtr<SWidget> FVariablesOutlinerMode::CreateContextMenu()
 					return;
 				}
 
-				FToolMenuSection& VariablesSection = InMenu->AddSection("Variables", LOCTEXT("VariablesSectionLabel", "Variables"));
+				const int32 NumSelectedVariables = MenuContext->WeakEntries.Num();
+				
+				/** Asset specific actions */
+				FToolMenuSection& AssetSection = InMenu->AddSection("Asset", LOCTEXT("AssetSectionLabel", "Asset"));
+				{
+					AssetSection.AddMenuEntryWithCommandList(FVariablesOutlinerCommands::Get().AddNewVariable, MenuContext->WeakCommandList.Pin(), TAttribute<FText>(), TAttribute<FText>(), FSlateIcon(FAppStyle::GetAppStyleSetName(), TEXT("Icons.Plus")));
+					AssetSection.AddMenuEntryWithCommandList(FVariablesOutlinerCommands::Get().AddNewVariables, MenuContext->WeakCommandList.Pin(), TAttribute<FText>(), TAttribute<FText>(), FSlateIcon(FAppStyle::GetAppStyleSetName(), TEXT("Icons.Plus")));
+					AssetSection.AddMenuEntryWithCommandList(FVariablesOutlinerCommands::Get().SaveAsset, MenuContext->WeakCommandList.Pin(), TAttribute<FText>(), TAttribute<FText>(), FSlateIcon(FAppStyle::Get().GetStyleSetName(), "AssetEditor.SaveAsset"));
+					AssetSection.AddMenuEntryWithCommandList(FGenericCommands::Get().Paste, MenuContext->WeakCommandList.Pin());
+				}
 
-				VariablesSection.AddMenuEntry(TEXT("AddVariables"),
-					LOCTEXT("AddVariablesMenuItem", "Add Variable(s)"),
-					LOCTEXT("AddVariablesMenuItemTooltip", "Adds variables to assets.\nIf multiple assets are selected, then variables will be added to each.\nIf no assets are selected and there are multiple assets, variables will be added to all assets."),
-					FSlateIcon(FAppStyle::Get().GetStyleSetName(), "Icons.Plus"),
-					FUIAction(FExecuteAction::CreateLambda([WeakOutliner = MenuContext->WeakOutliner]()
+				/** Variable specific actions */
+				FToolMenuSection& VariablesSection = InMenu->AddSection("Variables", LOCTEXT("VariablesSectionLabel", "Variables"));			
+				{
+					VariablesSection.AddSubMenu("FindReferencesMenu", LOCTEXT("FindReferencesMenuLabel", "Find References..."), LOCTEXT("FindReferencesMenuToolTip", "Find references to selected variables by different means."),
+					FNewToolMenuDelegate::CreateLambda(
+						[MenuContext](UToolMenu* InSubmenu)
+							{
+								InSubmenu->AddMenuEntry(NAME_None, FToolMenuEntry::InitMenuEntryWithCommandList(FVariablesOutlinerCommands::Get().FindReferences, MenuContext->WeakCommandList.Pin(), TAttribute<FText>(), TAttribute<FText>(), UAnimNextAssetFindReplaceVariables::GetIconFromSearchScope(ESearchScope::Global)));
+
+								InSubmenu->AddMenuEntry(NAME_None, FToolMenuEntry::InitMenuEntryWithCommandList(FVariablesOutlinerCommands::Get().FindReferencesInWorkspace, MenuContext->WeakCommandList.Pin(), TAttribute<FText>(), TAttribute<FText>(),
+									UAnimNextAssetFindReplaceVariables::GetIconFromSearchScope(ESearchScope::Workspace)));
+							
+								InSubmenu->AddMenuEntry(NAME_None, FToolMenuEntry::InitMenuEntryWithCommandList(FVariablesOutlinerCommands::Get().FindReferencesInAsset, MenuContext->WeakCommandList.Pin(), TAttribute<FText>(), TAttribute<FText>(),
+									UAnimNextAssetFindReplaceVariables::GetIconFromSearchScope(ESearchScope::Asset)));
+							}),
+						FToolUIActionChoice(),
+							EUserInterfaceActionType::Button,
+							false,
+							FSlateIcon(FAppStyle::GetAppStyleSetName(), TEXT("Icons.Search"))
+						);
+
+					VariablesSection.AddMenuEntryWithCommandList(FGenericCommands::Get().Rename, MenuContext->WeakCommandList.Pin());
+					VariablesSection.AddMenuEntryWithCommandList(FGenericCommands::Get().Delete, MenuContext->WeakCommandList.Pin());
+					VariablesSection.AddMenuEntryWithCommandList(FGenericCommands::Get().Duplicate, MenuContext->WeakCommandList.Pin());
+					
+					VariablesSection.AddMenuEntryWithCommandList(FGenericCommands::Get().Copy, MenuContext->WeakCommandList.Pin());
+
+					if (NumSelectedVariables == 1)
 					{
-						if (TSharedPtr<SVariablesOutliner> Outliner = WeakOutliner.Pin())
+						if (const UAnimNextVariableEntry* VariableEntry = Cast<UAnimNextVariableEntry>(MenuContext->WeakEntries[0].Get()))
 						{
-							Outliner->HandleAddVariablesClicked();
+							VariablesSection.AddMenuEntryWithCommandList(FVariablesOutlinerCommands::Get().ToggleVariableExport, MenuContext->WeakCommandList.Pin(), VariableEntry->GetExportAccessSpecifier() == EAnimNextExportAccessSpecifier::Private ? LOCTEXT("MakePublicMenuLabel", "Expose as Public") : LOCTEXT("MakePrivateMenuLabel", "Make Private"), TAttribute<FText>(), FSlateIcon(FAppStyle::Get().GetStyleSetName(), VariableEntry->GetExportAccessSpecifier() == EAnimNextExportAccessSpecifier::Private ? "Level.VisibleIcon16x" : "Level.NotVisibleHighlightIcon16x"));
 						}
-					}))
-				);
-
-				VariablesSection.AddMenuEntryWithCommandList(FGenericCommands::Get().Delete, MenuContext->WeakCommandList.Pin());
-				VariablesSection.AddMenuEntryWithCommandList(FGenericCommands::Get().Rename, MenuContext->WeakCommandList.Pin());
+					}
+				}
 
 				if (TSharedPtr<SVariablesOutliner> Outliner = MenuContext->WeakOutliner.Pin())
 				{
@@ -148,7 +182,7 @@ void FVariablesOutlinerMode::OnItemClicked(FSceneOutlinerTreeItemPtr Item)
 	HandleItemSelection(Selection);
 }
 
-void FVariablesOutlinerMode::HandleItemSelection(const FSceneOutlinerItemSelection& Selection)
+void FVariablesOutlinerMode::HandleItemSelection(const FSceneOutlinerItemSelection& Selection) const
 {
 	if (TSharedPtr<UE::Workspace::IWorkspaceEditor> WorkspaceEditor = WeakWorkspaceEditor.Pin())
 	{
@@ -165,11 +199,11 @@ void FVariablesOutlinerMode::HandleItemSelection(const FSceneOutlinerItemSelecti
 					EntriesToShow.Add(VariableEntry);
 				}
 			}
-			else if(const FVariablesOutlinerDataInterfaceItem* DataInterfaceItem = Item->CastTo<FVariablesOutlinerDataInterfaceItem>())
+			else if(const FVariablesOutlinerStructSharedVariablesItem* SharedVariablesItem = Item->CastTo<FVariablesOutlinerStructSharedVariablesItem>())
 			{
-				if(UAnimNextSharedVariablesEntry* SharedVariablesEntry = DataInterfaceItem->WeakEntry.Get())
+				if(const UAnimNextSharedVariablesEntry* SharedVariablesEntry = SharedVariablesItem->WeakEntry.Get())
 				{
-					EntriesToShow.Add(SharedVariablesEntry);
+					EntriesToShow.Add(const_cast<UAnimNextSharedVariablesEntry*>(SharedVariablesEntry));
 				}
 			}
 		}
@@ -179,22 +213,26 @@ void FVariablesOutlinerMode::HandleItemSelection(const FSceneOutlinerItemSelecti
 }
 
 void FVariablesOutlinerMode::OnItemSelectionChanged(FSceneOutlinerTreeItemPtr Item, ESelectInfo::Type SelectionType, const FSceneOutlinerItemSelection& Selection)
-{
-	HandleItemSelection(Selection);
-
+{	
 	if (TSharedPtr<UE::Workspace::IWorkspaceEditor> SharedWorkspaceEditor = WeakWorkspaceEditor.Pin())
 	{
 		SharedWorkspaceEditor->SetGlobalSelection(SceneOutliner->AsShared(), UE::Workspace::FOnClearGlobalSelection::CreateRaw(this, &FVariablesOutlinerMode::ResetOutlinerSelection));
 	}
+	
+	HandleItemSelection(Selection);
 }
 
 bool FVariablesOutlinerMode::CanRenameItem(const ISceneOutlinerTreeItem& Item) const
 {
 	if (const FVariablesOutlinerEntryItem* EntryItem = Item.CastTo<FVariablesOutlinerEntryItem>())
+	{
+		return !EntryItem->HasStructOwner();
+	}
+	else if (Item.IsA<FVariablesOutlinerAssetItem>())
 	{
 		return true;
 	}
-	else if (const FVariablesOutlinerAssetItem* AssetItem = Item.CastTo<FVariablesOutlinerAssetItem>())
+	else if (Item.IsA<FVariablesOutlinerCategoryItem>())
 	{
 		return true;
 	}
@@ -206,13 +244,99 @@ void FVariablesOutlinerMode::BindCommands(const TSharedRef<FUICommandList>& OutC
 	CommandList->MapAction( 
 		FGenericCommands::Get().Rename, 
 		FExecuteAction::CreateRaw(this, &FVariablesOutlinerMode::Rename),
-		FCanExecuteAction::CreateRaw(this, &FVariablesOutlinerMode::CanRename)
+		FCanExecuteAction::CreateRaw(this, &FVariablesOutlinerMode::CanRename),
+		FIsActionChecked(),
+		FIsActionButtonVisible::CreateRaw(this, &FVariablesOutlinerMode::CanRename)
 	);
 
 	CommandList->MapAction(
 		FGenericCommands::Get().Delete,
 		FExecuteAction::CreateRaw(this, &FVariablesOutlinerMode::Delete),
-		FCanExecuteAction::CreateRaw(this, &FVariablesOutlinerMode::CanDelete));
+		FCanExecuteAction::CreateRaw(this, &FVariablesOutlinerMode::CanDelete),
+		FIsActionChecked(),
+		FIsActionButtonVisible::CreateRaw(this, &FVariablesOutlinerMode::CanDelete)	
+		);
+
+	// [TODO] implement copy/paste for variables
+	CommandList->MapAction(
+		FGenericCommands::Get().Copy, 
+		FExecuteAction::CreateRaw(this, &FVariablesOutlinerMode::Copy),
+		FCanExecuteAction::CreateRaw(this, &FVariablesOutlinerMode::CanCopy),
+		FIsActionChecked(),
+		FIsActionButtonVisible::CreateRaw(this, &FVariablesOutlinerMode::CanCopy)
+	);
+
+	CommandList->MapAction( 
+		FGenericCommands::Get().Paste, 
+		FExecuteAction::CreateRaw(this, &FVariablesOutlinerMode::Paste),
+		FCanExecuteAction::CreateRaw(this, &FVariablesOutlinerMode::CanPaste),
+		FIsActionChecked(),
+		FIsActionButtonVisible::CreateRaw(this, &FVariablesOutlinerMode::CanPaste)
+	);
+
+	CommandList->MapAction(
+		FGenericCommands::Get().Duplicate,
+		FExecuteAction::CreateRaw(this, &FVariablesOutlinerMode::Duplicate),
+		FCanExecuteAction::CreateRaw(this, &FVariablesOutlinerMode::CanDuplicate),
+		FIsActionChecked(),
+		FIsActionButtonVisible::CreateRaw(this, &FVariablesOutlinerMode::CanDuplicate)
+	);
+
+	CommandList->MapAction(
+		FVariablesOutlinerCommands::Get().AddNewVariables,
+		FExecuteAction::CreateRaw(this, &FVariablesOutlinerMode::AddNewVariables),
+		FCanExecuteAction(),
+		FIsActionChecked(),
+		FIsActionButtonVisible::CreateRaw(this, &FVariablesOutlinerMode::CanAddVariable)
+		);
+
+	CommandList->MapAction(
+		FVariablesOutlinerCommands::Get().AddNewVariable,
+		FExecuteAction::CreateRaw(this, &FVariablesOutlinerMode::AddNewVariable),
+		FCanExecuteAction(),
+		FIsActionChecked(),
+		FIsActionButtonVisible::CreateRaw(this, &FVariablesOutlinerMode::CanAddVariable)
+		);
+
+	CommandList->MapAction(
+		FVariablesOutlinerCommands::Get().FindReferences,
+		FExecuteAction::CreateRaw(this, &FVariablesOutlinerMode::FindReferences, ESearchScope::Global),
+		FCanExecuteAction::CreateRaw(this, &FVariablesOutlinerMode::CanFindReferences),
+		FIsActionChecked(),
+		FIsActionButtonVisible::CreateRaw(this, &FVariablesOutlinerMode::IsFindReferencesVisible, ESearchScope::Global)
+	);
+	
+	CommandList->MapAction(
+		FVariablesOutlinerCommands::Get().FindReferencesInWorkspace,
+		FExecuteAction::CreateRaw(this, &FVariablesOutlinerMode::FindReferences, ESearchScope::Workspace),
+		FCanExecuteAction::CreateRaw(this, &FVariablesOutlinerMode::CanFindReferences),
+		FIsActionChecked(),
+	FIsActionButtonVisible::CreateRaw(this, &FVariablesOutlinerMode::IsFindReferencesVisible, ESearchScope::Workspace)
+	);
+	
+	CommandList->MapAction(
+		FVariablesOutlinerCommands::Get().FindReferencesInAsset,
+		FExecuteAction::CreateRaw(this, &FVariablesOutlinerMode::FindReferences, ESearchScope::Asset),
+		FCanExecuteAction::CreateRaw(this, &FVariablesOutlinerMode::CanFindReferences),
+		FIsActionChecked(),
+		FIsActionButtonVisible::CreateRaw(this, &FVariablesOutlinerMode::IsFindReferencesVisible, ESearchScope::Asset)
+	);
+
+	CommandList->MapAction(
+		FVariablesOutlinerCommands::Get().SaveAsset,
+		FExecuteAction::CreateRaw(this, &FVariablesOutlinerMode::SaveAsset),
+		FCanExecuteAction::CreateRaw(this, &FVariablesOutlinerMode::CanSaveAsset),
+		FIsActionChecked(),
+		FIsActionButtonVisible::CreateRaw(this, &FVariablesOutlinerMode::IsAsset)
+		);
+	
+	CommandList->MapAction(
+		FVariablesOutlinerCommands::Get().ToggleVariableExport,
+		FExecuteAction::CreateRaw(this, &FVariablesOutlinerMode::ToggleVariableExport),
+		FCanExecuteAction::CreateRaw(this, &FVariablesOutlinerMode::CanToggleVariableExport),
+		FIsActionChecked(),
+		FIsActionButtonVisible::CreateRaw(this, &FVariablesOutlinerMode::CanToggleVariableExport)
+		);
 }
 
 TUniquePtr<ISceneOutlinerHierarchy> FVariablesOutlinerMode::CreateHierarchy()
@@ -220,7 +344,7 @@ TUniquePtr<ISceneOutlinerHierarchy> FVariablesOutlinerMode::CreateHierarchy()
 	return MakeUnique<FVariablesOutlinerHierarchy>(this);
 }
 
-void FVariablesOutlinerMode::ResetOutlinerSelection()
+void FVariablesOutlinerMode::ResetOutlinerSelection() const
 {
 	SceneOutliner->ClearSelection();
 }
@@ -230,7 +354,7 @@ SVariablesOutliner* FVariablesOutlinerMode::GetOutliner() const
 	return static_cast<SVariablesOutliner*>(SceneOutliner);
 }
 
-void FVariablesOutlinerMode::Rename()
+void FVariablesOutlinerMode::Rename() const
 {
 	const FSceneOutlinerItemSelection& Selection = SceneOutliner->GetSelection();
 	if (Selection.Num() == 1)
@@ -256,7 +380,7 @@ bool FVariablesOutlinerMode::CanRename() const
 	return false;
 }
 
-void FVariablesOutlinerMode::Delete()
+void FVariablesOutlinerMode::Delete() const
 {
 	int32 NumEntries = 0;
 	TMap<UAnimNextRigVMAssetEditorData*, TArray<UAnimNextRigVMAssetEntry*>> EntriesToDeletePerAsset;
@@ -265,26 +389,29 @@ void FVariablesOutlinerMode::Delete()
 	{
 		if (const FVariablesOutlinerEntryItem* VariablesItem = Item->CastTo<FVariablesOutlinerEntryItem>())
 		{
-			UAnimNextVariableEntry* VariableEntry = VariablesItem->WeakEntry.Get();
-			UAnimNextSharedVariablesEntry* SharedVariablesEntry = VariablesItem->WeakDataInterfaceEntry.Get();
-			if(VariableEntry == nullptr || SharedVariablesEntry != nullptr)	// Cant delete variables in other data interfaces
+			if (!VariablesItem->HasStructOwner())
 			{
-				continue;
-			}
+				UAnimNextVariableEntry* VariableEntry = VariablesItem->WeakEntry.Get();
+				const UAnimNextSharedVariablesEntry* SharedVariablesEntry = VariablesItem->WeakSharedVariablesEntry.Get();
+				if(VariableEntry == nullptr || SharedVariablesEntry != nullptr)	// Cant delete variables in other data interfaces
+				{
+					continue;
+				}
 
-			UAnimNextRigVMAssetEditorData* EditorData = VariableEntry->GetTypedOuter<UAnimNextRigVMAssetEditorData>();
-			if(EditorData == nullptr)
-			{
-				continue;
-			}
+				UAnimNextRigVMAssetEditorData* EditorData = VariableEntry->GetTypedOuter<UAnimNextRigVMAssetEditorData>();
+				if(EditorData == nullptr)
+				{
+					continue;
+				}
 
-			TArray<UAnimNextRigVMAssetEntry*>& EntriesToDelete = EntriesToDeletePerAsset.FindOrAdd(EditorData);
-			EntriesToDelete.Add(VariableEntry);
-			NumEntries++;
+				TArray<UAnimNextRigVMAssetEntry*>& EntriesToDelete = EntriesToDeletePerAsset.FindOrAdd(EditorData);
+				EntriesToDelete.Add(VariableEntry);
+				NumEntries++;
+			}
 		}
-		else if (const FVariablesOutlinerDataInterfaceItem* DataInterfaceItem = Item->CastTo<FVariablesOutlinerDataInterfaceItem>())
+		else if (const FVariablesOutlinerStructSharedVariablesItem* SharedVariablesItem = Item->CastTo<FVariablesOutlinerStructSharedVariablesItem>())
 		{
-			UAnimNextSharedVariablesEntry* SharedVariablesEntry = DataInterfaceItem->WeakEntry.Get();
+			UAnimNextSharedVariablesEntry* SharedVariablesEntry = SharedVariablesItem->WeakEntry.Get();
 			if(SharedVariablesEntry == nullptr)
 			{
 				continue;
@@ -304,63 +431,500 @@ void FVariablesOutlinerMode::Delete()
 
 	if(NumEntries > 0)
 	{
+		// [TODO] prompt user when deletion impacts assets not part of the current workspace (local context), as these have a hidden (read unknown) impact
 		FScopedTransaction Transaction(FText::Format(LOCTEXT("DeleteVariablesFormat", "Delete {0}|plural(one=variable, other=variables)"), NumEntries));
 		for(const TPair<UAnimNextRigVMAssetEditorData*, TArray<UAnimNextRigVMAssetEntry*>>& EntriesPair : EntriesToDeletePerAsset)
 		{
-			EntriesPair.Key->RemoveEntries(EntriesPair.Value);
+			for (UAnimNextRigVMAssetEntry* Entry : EntriesPair.Value)
+			{
+				UncookedOnly::FUtils::DeleteVariable(CastChecked<UAnimNextVariableEntry>(Entry), true, true);
+			}
 		}
 	}
 }
 
 bool FVariablesOutlinerMode::CanDelete() const
 {
-	return true;
+	TArray<FSceneOutlinerTreeItemPtr> SelectedItems = SceneOutliner->GetSelectedItems();
+
+	bool bCanDelete = false;
+	for(const FSceneOutlinerTreeItemPtr& Item : SelectedItems)
+	{
+		if (const FVariablesOutlinerEntryItem* VariablesItem = Item->CastTo<FVariablesOutlinerEntryItem>())
+		{
+			if (!VariablesItem->HasStructOwner())
+			{
+				bCanDelete = true;
+			}
+		}
+	}
+	
+	return bCanDelete;
 }
 
-TSharedPtr<FDragDropOperation> FVariablesOutlinerMode::CreateDragDropOperation(const FPointerEvent& MouseEvent, const TArray<FSceneOutlinerTreeItemPtr>& InTreeItems) const
+void FVariablesOutlinerMode::Copy() const
 {
-	if(InTreeItems.Num() == 0)
+	const FSceneOutlinerItemSelection& Selection = SceneOutliner->GetSelection();
+	Selection.ForEachItem([](const FSceneOutlinerTreeItemPtr& SelectedItem)
 	{
-		return nullptr;
+		
+	});
+}
+
+bool FVariablesOutlinerMode::CanCopy() const
+{
+	// [TODO] implement copy behaviour
+	return false;
+}
+
+void FVariablesOutlinerMode::Paste() const
+{
+}
+
+bool FVariablesOutlinerMode::CanPaste() const
+{
+	// [TODO] implement paste behaviour
+	return false;
+}
+
+void FVariablesOutlinerMode::ToggleVariableExport() const
+{
+	const FSceneOutlinerItemSelection& Selection = SceneOutliner->GetSelection();
+	if (Selection.Num())
+	{		
+		FScopedTransaction Transaction(LOCTEXT("SetAccessSpecifierTransaction", "Set Access Specifier"));
+		Selection.ForEachItem<FVariablesOutlinerEntryItem>([](const FVariablesOutlinerEntryItem& EntryItem)
+		{			
+			if (IAnimNextRigVMExportInterface* Export = Cast<IAnimNextRigVMExportInterface>(EntryItem.WeakEntry.Get()))
+			{
+				if (UAnimNextRigVMAsset* OuterAsset = EntryItem.WeakEntry.Get()->GetTypedOuter<UAnimNextRigVMAsset>())
+				{
+					if (OuterAsset->GetClass() && OuterAsset->GetClass() != UAnimNextSharedVariables::StaticClass())
+					{
+						Export->SetExportAccessSpecifier(Export->GetExportAccessSpecifier() == EAnimNextExportAccessSpecifier::Public ? EAnimNextExportAccessSpecifier::Private : EAnimNextExportAccessSpecifier::Public);
+					}
+				}
+			}
+		});
 	}
+}
 
-	FVariablesOutlinerEntryItem* VariableItem = InTreeItems[0]->CastTo<FVariablesOutlinerEntryItem>();
-	if(VariableItem == nullptr)
+bool FVariablesOutlinerMode::CanToggleVariableExport() const
+{
+	const FSceneOutlinerItemSelection& Selection = SceneOutliner->GetSelection();
+	if (Selection.Num())
 	{
-		return nullptr;
+		bool bCanToggle = false;
+		
+		FScopedTransaction Transaction(LOCTEXT("SetAccessSpecifierTransaction", "Set Access Specifier"));
+		Selection.ForEachItem<FVariablesOutlinerEntryItem>([&bCanToggle](const FVariablesOutlinerEntryItem& EntryItem)
+		{
+			if (UAnimNextVariableEntry* VariableEntry = Cast<UAnimNextVariableEntry>(EntryItem.WeakEntry.Get()))
+			{
+				if (UAnimNextRigVMAsset* OuterAsset = VariableEntry->GetTypedOuter<UAnimNextRigVMAsset>())
+				{
+					// Disable toggling public/private for pure UAnimNextSharedVariables objects				
+					if (OuterAsset->GetClass() && OuterAsset->GetClass() != UAnimNextSharedVariables::StaticClass())
+					{
+						bCanToggle = true;
+					}
+				}
+			}
+		});
+
+		return bCanToggle;
 	}
 
-	if(UAnimNextVariableEntry* Entry = Cast<UAnimNextVariableEntry>(VariableItem->WeakEntry.Get()))
+	return false;
+}
+
+void FVariablesOutlinerMode::Duplicate() const
+{	
+	const FSceneOutlinerItemSelection& Selection = SceneOutliner->GetSelection();
+	if (Selection.Num())
+	{		
+		FScopedTransaction Transaction(FText::Format(LOCTEXT("DuplicateVariableFormat", "Duplicate {0}|plural(one=variable, other=variables)"), Selection.Num()));
+
+		Selection.ForEachItem<FVariablesOutlinerEntryItem>([](const FVariablesOutlinerEntryItem& EntryItem)
+		{
+			if (const UAnimNextVariableEntry* VariableEntry = EntryItem.WeakEntry.Get())
+			{
+				if (UAnimNextRigVMAssetEditorData* EditorData = VariableEntry->GetTypedOuter<UAnimNextRigVMAssetEditorData>())
+				{	
+					const FName VariableName = UncookedOnly::FUtils::GetValidVariableName(EditorData, VariableEntry->GetVariableName());		
+					EditorData->AddVariable(VariableName, VariableEntry->GetType());
+				}
+			}
+		});
+	}
+}	
+
+bool FVariablesOutlinerMode::CanDuplicate() const
+{
+	const FSceneOutlinerItemSelection& Selection = SceneOutliner->GetSelection();
+
+	bool bCanDuplicate = true;
+	Selection.ForEachItem([&bCanDuplicate](const FSceneOutlinerTreeItemPtr& Item)
 	{
-		const UAnimNextSharedVariables* Asset = nullptr;
-		if (VariableItem->WeakDataInterfaceEntry.Get())
+		const FVariablesOutlinerEntryItem* EntryItem = Item->CastTo<FVariablesOutlinerEntryItem>();
+		if (EntryItem == nullptr || EntryItem->HasStructOwner())
 		{
-			Asset = VariableItem->WeakDataInterfaceEntry->GetAsset();
+			bCanDuplicate = false;
+		}
+	});
+	
+	return bCanDuplicate;
+}
+
+void FVariablesOutlinerMode::FindReferences(ESearchScope InSearchScope) const
+{
+	const FSceneOutlinerItemSelection& Selection = SceneOutliner->GetSelection();
+	if (Selection.Num() == 1)
+	{
+		FSceneOutlinerTreeItemPtr ItemToFind = Selection.SelectedItems[0].Pin();
+		if (const FVariablesOutlinerEntryItem* EntryItem = ItemToFind->CastTo<FVariablesOutlinerEntryItem>())
+		{
+			if (TSharedPtr<SDockTab> FindAndReplaceTab = WeakWorkspaceEditor.Pin()->GetTabManager()->TryInvokeTab(FTabId(UE::AnimNext::Editor::FindAndReplaceTabName)))
+			{
+				TSharedRef<IAnimAssetFindReplace> FindAndReplace = StaticCastSharedRef<IAnimAssetFindReplace>(FindAndReplaceTab->GetContent());
+				FindAndReplace->SetCurrentProcessor(UAnimNextAssetFindReplaceVariables::StaticClass());
+				UAnimNextAssetFindReplaceVariables* AnimNextFindAndReplaceProcessor = FindAndReplace->GetProcessor<UAnimNextAssetFindReplaceVariables>();
+
+				if (TSharedPtr<UE::Workspace::IWorkspaceEditor> SharedWorkspaceEditor = WeakWorkspaceEditor.Pin())
+				{
+					AnimNextFindAndReplaceProcessor->SetWorkspaceEditor(SharedWorkspaceEditor.ToSharedRef());
+				}
+				
+				UAnimNextVariableEntry* VariableEntry = EntryItem->WeakEntry.Get();
+				AnimNextFindAndReplaceProcessor->SetSearchScope(InSearchScope);		
+				AnimNextFindAndReplaceProcessor->SetFindReferenceFromEntry(VariableEntry);				
+			}
+		}		
+	}
+}
+
+bool FVariablesOutlinerMode::CanFindReferences() const
+{
+	const FSceneOutlinerItemSelection& Selection = SceneOutliner->GetSelection();
+	if (Selection.Num() == 1)
+	{
+		FSceneOutlinerTreeItemPtr ItemToFind = Selection.SelectedItems[0].Pin();
+		if (const FVariablesOutlinerEntryItem* EntryItem = ItemToFind->CastTo<FVariablesOutlinerEntryItem>())
+		{
+			return true;
+		}
+	}
+	
+	return false;
+}
+
+bool FVariablesOutlinerMode::IsFindReferencesVisible(ESearchScope InSearchScope) const
+{
+	return CanFindReferences();
+}
+
+bool FVariablesOutlinerMode::CanAddVariable() const
+{
+	const FSceneOutlinerItemSelection& Selection = SceneOutliner->GetSelection();
+	if (Selection.Num() == 1)
+	{
+		const FSceneOutlinerTreeItemPtr ItemToFind = Selection.SelectedItems[0].Pin();
+		return ItemToFind->IsA<FVariablesOutlinerAssetItem>() || ItemToFind->IsA<FVariablesOutlinerCategoryItem>();
+	}
+	
+	return false;
+}
+
+void FVariablesOutlinerMode::AddNewVariable(UAnimNextRigVMAsset* InAsset)
+{
+	if (UAnimNextRigVMAssetEditorData* EditorData = UncookedOnly::FUtils::GetEditorData<UAnimNextRigVMAssetEditorData>(InAsset))
+	{
+		UAnimNextVariableSettings* Settings = GetMutableDefault<UAnimNextVariableSettings>();
+		FScopedTransaction Transaction(LOCTEXT("AddVariable", "Add variable"));
+		
+		const FName VariableName = UncookedOnly::FUtils::GetValidVariableName(EditorData, Settings->GetLastVariableName());		
+		EditorData->AddVariable(VariableName, Settings->GetLastVariableType());
+		Settings->SetLastVariableName(VariableName);
+	}
+}
+
+void FVariablesOutlinerMode::SetHighlightedItem(FSceneOutlinerTreeItemID InID) const
+{
+	if (SVariablesOutliner* Outliner = static_cast<SVariablesOutliner*>(SceneOutliner))
+	{
+		if (FSceneOutlinerTreeItemPtr ItemPtr = Outliner->GetTreeItem(InID))
+		{
+			Outliner->FrameItem(InID);
+			Outliner->SetHighlightedItem(ItemPtr);
+		}
+	}
+}
+
+void FVariablesOutlinerMode::ClearHighlightedItem(FSceneOutlinerTreeItemID InID) const
+{
+	if (SVariablesOutliner* Outliner = static_cast<SVariablesOutliner*>(SceneOutliner))
+	{
+		if (FSceneOutlinerTreeItemPtr ItemPtr = Outliner->GetTreeItem(InID))
+		{
+			Outliner->ClearHighlightedItem(ItemPtr);
+		}
+	}
+}
+
+void FVariablesOutlinerMode::AddNewVariable() const
+{
+	TArray<FSceneOutlinerTreeItemPtr> SelectedItems = SceneOutliner->GetSelectedItems();
+	
+	TMap<UAnimNextRigVMAssetEditorData*, FString> AssetsAndCategoriesToAddTo;
+	for(const FSceneOutlinerTreeItemPtr& Item : SelectedItems)
+	{
+		const FVariablesOutlinerAssetItem* AssetItem = Item->CastTo<FVariablesOutlinerAssetItem>();
+		const FVariablesOutlinerCategoryItem* CategoryItem = Item->CastTo<FVariablesOutlinerCategoryItem>();
+		if (AssetItem == nullptr && CategoryItem == nullptr)
+		{			
+			continue;
+		}
+
+		UAnimNextRigVMAsset* Asset = nullptr;
+		FString Category;
+		
+		if (CategoryItem)
+		{
+			Asset = CategoryItem->WeakOwner.Get();
+			Category = CategoryItem->GetDisplayString();
 		}
 		else
 		{
-			Asset = Entry->GetTypedOuter<UAnimNextSharedVariables>();
+			Asset = AssetItem->SoftAsset.Get();
+		}
+
+		if(Asset == nullptr)
+		{
+			continue;
+		}
+
+		UAnimNextRigVMAssetEditorData* EditorData = UncookedOnly::FUtils::GetEditorData<UAnimNextRigVMAssetEditorData>(Asset);
+		if(EditorData == nullptr)
+		{
+			continue;
+		}
+
+		AssetsAndCategoriesToAddTo.Add(EditorData, Category);
+	}
+	
+	UAnimNextVariableSettings* Settings = GetMutableDefault<UAnimNextVariableSettings>();
+	FScopedTransaction Transaction(LOCTEXT("AddVariable", "Add variable"));
+	for(const TPair<UAnimNextRigVMAssetEditorData*, FString>& Pair : AssetsAndCategoriesToAddTo)
+	{
+		const FName VariableName = UncookedOnly::FUtils::GetValidVariableName(Pair.Key, Settings->GetLastVariableName());
+		Settings->SetLastVariableName(VariableName);
+		
+		if (UAnimNextVariableEntry* VariableEntry = Pair.Key->AddVariable(VariableName, Settings->GetLastVariableType()))
+		{
+			VariableEntry->SetVariableCategory(Pair.Value);
+		}
+	}
+}
+
+void FVariablesOutlinerMode::AddNewVariables() const
+{
+	TArray<FSceneOutlinerTreeItemPtr> SelectedItems = SceneOutliner->GetSelectedItems();
+	
+	TArray<UAnimNextRigVMAssetEditorData*> AssetsToAddTo;
+	for(const FSceneOutlinerTreeItemPtr& Item : SelectedItems)
+	{
+		const FVariablesOutlinerAssetItem* AssetItem = Item->CastTo<FVariablesOutlinerAssetItem>();
+		if (AssetItem == nullptr)
+		{
+			continue;
+		}
+
+		UAnimNextRigVMAsset* Asset = AssetItem->SoftAsset.Get();
+		if(Asset == nullptr)
+		{
+			continue;
+		}
+
+		UAnimNextRigVMAssetEditorData* EditorData = UncookedOnly::FUtils::GetEditorData<UAnimNextRigVMAssetEditorData>(Asset);
+		if(EditorData == nullptr)
+		{
+			continue;
+		}
+			
+		AssetsToAddTo.Add(EditorData);
+	}
+	
+	TSharedRef<SAddVariablesDialog> AddVariablesDialog = SNew(SAddVariablesDialog, AssetsToAddTo);
+
+	TArray<FVariableToAdd> VariablesToAdd;
+	TArray<FSharedVariablesToAdd> AllSharedVariablesToAdd;
+	if(AddVariablesDialog->ShowModal(VariablesToAdd, AllSharedVariablesToAdd))
+	{
+		FScopedTransaction Transaction(FText::Format(LOCTEXT("AddVariablesFormat", "Add {0}|plural(one=variable, other=variables)"), (AllSharedVariablesToAdd.Num() + VariablesToAdd.Num()) * AssetsToAddTo.Num()));
+		for(UAnimNextRigVMAssetEditorData* EditorData : AssetsToAddTo)
+		{
+			for (const FVariableToAdd& VariableToAdd : VariablesToAdd)
+			{
+				EditorData->AddVariable(VariableToAdd.Name, VariableToAdd.Type);
+			}
+
+			for (const FSharedVariablesToAdd& SharedVariablesToAdd : AllSharedVariablesToAdd)
+			{
+				if (SharedVariablesToAdd.SharedVariables.IsType<UAnimNextSharedVariables*>())
+				{
+					EditorData->AddSharedVariables(SharedVariablesToAdd.SharedVariables.Get<UAnimNextSharedVariables*>());
+				}
+				else
+				{
+					EditorData->AddSharedVariablesStruct(SharedVariablesToAdd.SharedVariables.Get<const UScriptStruct*>());
+				}
+			}
+		}
+	}
+}
+
+void FVariablesOutlinerMode::SaveAsset() const
+{
+	const FSceneOutlinerItemSelection& Selection = SceneOutliner->GetSelection();
+	Selection.ForEachItem<FVariablesOutlinerAssetItem>([](const FVariablesOutlinerAssetItem& AssetItem)
+	{
+		if (const UAnimNextRigVMAsset* Asset = AssetItem.SoftAsset.Get())
+		{
+			FEditorFileUtils::PromptForCheckoutAndSave({Asset->GetPackage()}, false, /*bPromptToSave=*/ false);
 		}
+	});	
+}
 
-		if (Asset)
+bool FVariablesOutlinerMode::CanSaveAsset() const
+{
+	const FSceneOutlinerItemSelection& Selection = SceneOutliner->GetSelection();	
+	return Selection.SelectedItems.ContainsByPredicate([](const TWeakPtr<ISceneOutlinerTreeItem>& WeakItem)
+	{
+		if (WeakItem.IsValid())
 		{
-			TSharedPtr<FAnimNextSchemaAction_Variable> Action = MakeShared<FAnimNextSchemaAction_Variable>(Entry->GetVariableName(), Asset, Entry->GetType(), FAnimNextSchemaAction_Variable::EVariableAccessorChoice::Deferred);
-			return FVariableDragDropOp::New(Action);
+			if (const FVariablesOutlinerAssetItem* AssetItem = WeakItem.Pin()->CastTo<FVariablesOutlinerAssetItem>())
+			{
+				if (const UAnimNextRigVMAsset* Asset = AssetItem->SoftAsset.Get())
+				{
+					return Asset->GetPackage()->IsDirty();
+				}
+			}
 		}
+
+		return false;
+	});
+}
+
+bool FVariablesOutlinerMode::IsAsset() const
+{
+	const FSceneOutlinerItemSelection& Selection = SceneOutliner->GetSelection();	
+	return Selection.SelectedItems.ContainsByPredicate([](const TWeakPtr<ISceneOutlinerTreeItem>& WeakItem) { return WeakItem.IsValid() && WeakItem.Pin()->IsA<FVariablesOutlinerAssetItem>(); });	
+}
+
+TSharedPtr<FDragDropOperation> FVariablesOutlinerMode::CreateDragDropOperation(const FPointerEvent& MouseEvent, const TArray<FSceneOutlinerTreeItemPtr>& InTreeItems) const
+{
+	if(InTreeItems.Num() != 1)
+	{
+		return nullptr;
 	}
-	else if (const FProperty* Property = VariableItem->PropertyPath.Get())
+
+	const FSceneOutlinerTreeItemPtr& SelectedItem = InTreeItems[0];
+	if (FVariablesOutlinerEntryItem* VariableItem = SelectedItem->CastTo<FVariablesOutlinerEntryItem>())
 	{
-		const UScriptStruct* Struct = Property->GetOwner<UScriptStruct>();
-		if (Struct)
+		if(UAnimNextVariableEntry* Entry = Cast<UAnimNextVariableEntry>(VariableItem->WeakEntry.Get()))
 		{
-			TSharedPtr<FAnimNextSchemaAction_Variable> Action = MakeShared<FAnimNextSchemaAction_Variable>(Property->GetFName(), Struct, FAnimNextParamType::FromProperty(Property), FAnimNextSchemaAction_Variable::EVariableAccessorChoice::Deferred);
-			return FVariableDragDropOp::New(Action);
+			const UAnimNextSharedVariables* Asset = nullptr;
+			if (VariableItem->WeakSharedVariablesEntry.Get())
+			{
+				Asset = VariableItem->WeakSharedVariablesEntry->GetAsset();
+			}
+			else
+			{
+				Asset = Entry->GetTypedOuter<UAnimNextSharedVariables>();
+			}
+		
+			if (Asset)
+			{
+				TSharedPtr<FAnimNextSchemaAction_Variable> Action = MakeShared<FAnimNextSchemaAction_Variable>(Entry->GetVariableName(), Asset, Entry->GetType(), FAnimNextSchemaAction_Variable::EVariableAccessorChoice::Deferred);
+
+				return FVariableDragDropOp::New(Action, StaticCastSharedRef<FVariablesOutlinerEntryItem>(VariableItem->AsShared()));
+			}
+		}
+		else if (const FProperty* Property = VariableItem->PropertyPath.Get())
+		{
+			if (const UScriptStruct* Struct = Property->GetOwner<UScriptStruct>())
+			{
+				TSharedPtr<FAnimNextSchemaAction_Variable> Action = MakeShared<FAnimNextSchemaAction_Variable>(Property->GetFName(), Struct, FAnimNextParamType::FromProperty(Property), FAnimNextSchemaAction_Variable::EVariableAccessorChoice::Deferred);
+				
+				TSharedPtr<FVariableDragDropOp> Op = FVariableDragDropOp::New(Action, StaticCastSharedRef<FVariablesOutlinerEntryItem>(VariableItem->AsShared()));
+				return Op;
+			}
 		}
 	}
+	else if (FVariablesOutlinerCategoryItem* CategoryItem = SelectedItem->CastTo<FVariablesOutlinerCategoryItem>())
+	{
+		TSharedPtr<FCategoryDragDropOp> Op = FCategoryDragDropOp::New(StaticCastSharedRef<FVariablesOutlinerCategoryItem>(CategoryItem->AsShared()));
+		return Op;
+	}
 
 	return nullptr;
+
+}
+
+bool FVariablesOutlinerMode::ParseDragDrop(FSceneOutlinerDragDropPayload& OutPayload, const FDragDropOperation& Operation) const
+{
+	return ISceneOutlinerMode::ParseDragDrop(OutPayload, Operation);
 }
 
+FSceneOutlinerDragValidationInfo FVariablesOutlinerMode::ValidateDrop(const ISceneOutlinerTreeItem& DropTarget,
+	const FSceneOutlinerDragDropPayload& Payload) const
+{
+	return ISceneOutlinerMode::ValidateDrop(DropTarget, Payload);
+}
+
+void FVariablesOutlinerMode::OnDrop(ISceneOutlinerTreeItem& DropTarget, const FSceneOutlinerDragDropPayload& Payload,
+	const FSceneOutlinerDragValidationInfo& ValidationInfo) const
+{
+	ISceneOutlinerMode::OnDrop(DropTarget, Payload, ValidationInfo);
+}
+
+FReply FVariablesOutlinerMode::OnDragOverItem(const FDragDropEvent& Event, const ISceneOutlinerTreeItem& Item) const
+{
+	return ISceneOutlinerMode::OnDragOverItem(Event, Item);
+}
+
+int32 FVariablesOutlinerMode::GetTypeSortPriority(const ISceneOutlinerTreeItem& Item) const
+{
+	if(const FVariablesOutlinerAssetItem* AssetItem = Item.CastTo<FVariablesOutlinerAssetItem>())
+	{
+		return AssetItem->SortValue;
+	}
+
+	if(const FVariablesOutlinerEntryItem* EntryItem = Item.CastTo<FVariablesOutlinerEntryItem>())
+	{
+		return EntryItem->SortValue;
+	}
+
+	if(const FVariablesOutlinerCategoryItem* CategoryItem = Item.CastTo<FVariablesOutlinerCategoryItem>())
+	{
+		return CategoryItem->SortValue;
+	}
+	
+	return ISceneOutlinerMode::GetTypeSortPriority(Item);
+}
+
+void FVariablesOutlinerMode::OnItemDoubleClick(FSceneOutlinerTreeItemPtr Item)
+{
+	if (Item.IsValid())
+	{
+		if(FVariablesOutlinerAssetItem* AssetItem = Item->CastTo<FVariablesOutlinerAssetItem>())
+		{
+			// [TODO] Open asset from Variables Outliner
+		}
+	}	
+	
+	ISceneOutlinerMode::OnItemDoubleClick(Item);
+}
 }
 
 #undef LOCTEXT_NAMESPACE // "FVariablesOutlinerMode"
\ No newline at end of file
diff --git a/Engine/Plugins/Experimental/AnimNext/Source/AnimNextEditor/Private/Variables/Outliner/VariablesOutlinerMode.h b/Engine/Plugins/Experimental/AnimNext/Source/AnimNextEditor/Private/Variables/Outliner/VariablesOutlinerMode.h
--- a/Engine/Plugins/Experimental/AnimNext/Source/AnimNextEditor/Private/Variables/Outliner/VariablesOutlinerMode.h
+++ b/Engine/Plugins/Experimental/AnimNext/Source/AnimNextEditor/Private/Variables/Outliner/VariablesOutlinerMode.h
@@ -4,6 +4,7 @@
 
 #include "AnimNextRigVMAssetEditorData.h"
 #include "ISceneOutlinerMode.h"
+#include "Common/AnimNextAssetFindReplaceVariables.h"
 
 class UAnimNextRigVMAssetEditorData;
 class FVariablesOutlinerMode;
@@ -29,29 +30,62 @@ public:
 	virtual TSharedPtr<SWidget> CreateContextMenu() override;
 	virtual FReply OnKeyDown(const FKeyEvent& InKeyEvent) override;
 	virtual void OnItemClicked(FSceneOutlinerTreeItemPtr Item) override;
-	void HandleItemSelection(const FSceneOutlinerItemSelection& Selection);
+	void HandleItemSelection(const FSceneOutlinerItemSelection& Selection) const;
 	virtual void OnItemSelectionChanged(FSceneOutlinerTreeItemPtr Item, ESelectInfo::Type SelectionType, const FSceneOutlinerItemSelection& Selection) override;
 	virtual bool CanCustomizeToolbar() const override { return true; }
 	virtual bool CanRenameItem(const ISceneOutlinerTreeItem& Item) const override;
 	virtual ESelectionMode::Type GetSelectionMode() const override { return ESelectionMode::Multi; }
 	virtual void BindCommands(const TSharedRef<FUICommandList>& OutCommandList) override;
+	virtual int32 GetTypeSortPriority(const ISceneOutlinerTreeItem& Item) const override;
 	virtual TSharedPtr<FDragDropOperation> CreateDragDropOperation(const FPointerEvent& MouseEvent, const TArray<FSceneOutlinerTreeItemPtr>& InTreeItems) const override;
-
+	virtual bool ParseDragDrop(FSceneOutlinerDragDropPayload& OutPayload, const FDragDropOperation& Operation) const override;
+	virtual FSceneOutlinerDragValidationInfo ValidateDrop(const ISceneOutlinerTreeItem& DropTarget, const FSceneOutlinerDragDropPayload& Payload) const override;
+	virtual void OnDrop(ISceneOutlinerTreeItem& DropTarget, const FSceneOutlinerDragDropPayload& Payload, const FSceneOutlinerDragValidationInfo& ValidationInfo) const override;
+	virtual FReply OnDragOverItem(const FDragDropEvent& Event, const ISceneOutlinerTreeItem& Item) const override;
+	virtual bool CanSupportDragAndDrop() const override { return true; }
+	virtual void OnItemDoubleClick(FSceneOutlinerTreeItemPtr Item) override;
+	
+	static void AddNewVariable(UAnimNextRigVMAsset* InAsset);
+
+	void SetHighlightedItem(FSceneOutlinerTreeItemID InID) const;
+	void ClearHighlightedItem(FSceneOutlinerTreeItemID InID) const;
 protected:
 	virtual TUniquePtr<ISceneOutlinerHierarchy> CreateHierarchy() override;
 	// End ISceneOutlinerMode overrides
 
-	void ResetOutlinerSelection();
+	void ResetOutlinerSelection() const;
 
 	SVariablesOutliner* GetOutliner() const;
 
-	void Rename();
+	void Rename() const;
+	virtual bool CanRename() const override;
+
+	void Delete() const;
+	virtual bool CanDelete() const override;
+
+	void Copy() const;
+	virtual bool CanCopy() const override;
+
+	void Paste() const;
+	virtual bool CanPaste() const override;
 
-	bool CanRename() const;
+	void ToggleVariableExport() const;
+	bool CanToggleVariableExport() const;
 
-	void Delete();
+	void Duplicate() const;
+	bool CanDuplicate() const;
 
-	bool CanDelete() const;
+	void FindReferences(ESearchScope InSearchScope) const;
+	bool CanFindReferences() const;
+	bool IsFindReferencesVisible(ESearchScope InSearchScope) const;
+
+	bool CanAddVariable() const;
+	void AddNewVariable() const;
+	void AddNewVariables() const;
+
+	void SaveAsset() const;
+	bool CanSaveAsset() const;
+	bool IsAsset() const;
 
 private:
 	friend class FVariablesOutlinerHierarchy;
@@ -61,4 +95,4 @@ private:
 	TSharedPtr<FUICommandList> CommandList;
 };
 
-}
\ No newline at end of file
+}
diff --git a/Engine/Plugins/Experimental/AnimNext/Source/AnimNextEditor/Private/Variables/Outliner/VariablesOutlinerStructSharedVariablesItem.cpp b/Engine/Plugins/Experimental/AnimNext/Source/AnimNextEditor/Private/Variables/Outliner/VariablesOutlinerStructSharedVariablesItem.cpp
new file mode 100644
--- /dev/null
+++ b/Engine/Plugins/Experimental/AnimNext/Source/AnimNextEditor/Private/Variables/Outliner/VariablesOutlinerStructSharedVariablesItem.cpp
@@ -0,0 +1,159 @@
+﻿// Copyright Epic Games, Inc. All Rights Reserved.
+
+#include "VariablesOutlinerStructSharedVariablesItem.h"
+
+#include "ISceneOutliner.h"
+#include "Styling/SlateColor.h"
+#include "Misc/Attribute.h"
+#include "Misc/Optional.h"
+#include "SceneOutlinerStandaloneTypes.h"
+#include "AnimNextEditorModule.h"
+#include "AnimNextRigVMAssetEditorData.h"
+#include "VariablesOutlinerDragDrop.h"
+#include "Entries/AnimNextSharedVariablesEntry.h"
+#include "Styling/SlateIconFinder.h"
+#include "Styling/StarshipCoreStyle.h"
+#include "UObject/Package.h"
+#include "Widgets/SCompoundWidget.h"
+#include "Widgets/DeclarativeSyntaxSupport.h"
+#include "Widgets/Images/SImage.h"
+#include "Widgets/Text/SInlineEditableTextBlock.h"
+#include "Widgets/Text/STextBlock.h"
+#include "Widgets/Views/STableRow.h"
+#include "Widgets/Views/SListView.h"
+
+#define LOCTEXT_NAMESPACE "VariablesOutlinerTreeItem"
+
+namespace UE::AnimNext::Editor
+{
+
+const FSceneOutlinerTreeItemType FVariablesOutlinerStructSharedVariablesItem::Type;
+
+class SVariablesOutlinerStructSharedVariablesLabel : FSceneOutlinerCommonLabelData, public SCompoundWidget
+{
+	SLATE_BEGIN_ARGS(SVariablesOutlinerStructSharedVariablesLabel) {}
+	SLATE_END_ARGS()
+
+	void Construct(const FArguments& InArgs, FVariablesOutlinerStructSharedVariablesItem& InTreeItem, ISceneOutliner& SceneOutliner, const STableRow<FSceneOutlinerTreeItemPtr>& InRow)
+	{
+		WeakSceneOutliner = StaticCastSharedRef<ISceneOutliner>(SceneOutliner.AsShared());
+		TreeItem = StaticCastSharedRef<FVariablesOutlinerStructSharedVariablesItem>(InTreeItem.AsShared());
+
+		FText AssetName;
+		if(const UAnimNextSharedVariablesEntry* AssetEntry = InTreeItem.WeakEntry.Get())
+		{
+			AssetName = FText::FromString(AssetEntry->GetObjectPath().ToString());
+		}
+		else
+		{
+			AssetName = LOCTEXT("UnknownAssetName", "Unknown Asset");
+		}
+		
+		ChildSlot
+		[
+			SNew(SHorizontalBox)
+			.ToolTipText(FText::Format(LOCTEXT("ImportedVariablesFormat", "Shared variables from '{0}'"), AssetName))
+			+SHorizontalBox::Slot()
+			.AutoWidth()
+			.VAlign(VAlign_Center)
+			.Padding(2.0f, 2.0f)
+			[
+				SNew(SImage)
+				.Image(FSlateIconFinder::FindIconBrushForClass(UUserDefinedStruct::StaticClass()))
+			]
+			+SHorizontalBox::Slot()
+			.AutoWidth()
+			.VAlign(VAlign_Center)
+			.Padding(2.0f, 2.0f)
+			[
+				SNew(STextBlock)
+				.Font(FStyleFonts::Get().NormalBold)
+				.Text(this, &SVariablesOutlinerStructSharedVariablesLabel::GetDisplayText)
+				.HighlightText(SceneOutliner.GetFilterHighlightText())
+				.ColorAndOpacity(this, &SVariablesOutlinerStructSharedVariablesLabel::GetForegroundColor)
+			]
+		];
+	}
+	
+	FText GetDisplayText() const
+	{
+		if (const TSharedPtr<FVariablesOutlinerStructSharedVariablesItem> Item = TreeItem.Pin())
+		{
+			return FText::FromString(Item->GetDisplayString());
+		}
+		return FText();
+	}
+
+	virtual void OnDragEnter( const FGeometry& MyGeometry, const FDragDropEvent& DragDropEvent ) override
+	{
+		TSharedPtr<FVariableDragDropOp> GraphDropOp = DragDropEvent.GetOperationAs<FVariableDragDropOp>();
+		if (GraphDropOp.IsValid())
+		{
+			const FSlateBrush* ErrorSymbol = FAppStyle::GetBrush(TEXT("Graph.ConnectorFeedback.Error"));
+			GraphDropOp->SetSimpleFeedbackMessage(ErrorSymbol, FSlateColor(GetForegroundColor()), LOCTEXT("StructSharedVariablesFeedback", "Cannot make changes to Struct SharedVariables"));
+		}
+	}
+
+	virtual FSlateColor GetForegroundColor() const override
+	{
+		const TOptional<FLinearColor> BaseColor = FSceneOutlinerCommonLabelData::GetForegroundColor(*TreeItem.Pin());
+		return BaseColor.IsSet() ? BaseColor.GetValue() : FSlateColor::UseForeground();
+	}
+
+	TWeakPtr<FVariablesOutlinerStructSharedVariablesItem> TreeItem;
+	TSharedPtr<SInlineEditableTextBlock> TextBlock;
+};
+
+FVariablesOutlinerStructSharedVariablesItem::FVariablesOutlinerStructSharedVariablesItem(UAnimNextSharedVariablesEntry* InEntry)
+	: ISceneOutlinerTreeItem(FVariablesOutlinerStructSharedVariablesItem::Type)
+	, WeakEntry(InEntry)
+{
+}
+
+bool FVariablesOutlinerStructSharedVariablesItem::IsValid() const
+{
+	return WeakEntry.Get() != nullptr;
+}
+
+FSceneOutlinerTreeItemID FVariablesOutlinerStructSharedVariablesItem::GetID() const
+{
+	if (const UAnimNextSharedVariablesEntry* SharedVariableEntry = WeakEntry.Get())
+	{
+		check(SharedVariableEntry->GetStruct());
+		const FSoftObjectPath Path = SharedVariableEntry->GetStruct();
+		return GetTypeHash(Path);
+	}
+	
+	return GetTypeHash(WeakEntry);
+}
+
+FString FVariablesOutlinerStructSharedVariablesItem::GetDisplayString() const
+{
+	const UAnimNextRigVMAssetEntry* Entry = WeakEntry.Get();
+	if(Entry == nullptr)
+	{
+		return FString();
+	}
+
+	return Entry->GetDisplayName().ToString();
+}
+
+TSharedRef<SWidget> FVariablesOutlinerStructSharedVariablesItem::GenerateLabelWidget(ISceneOutliner& Outliner, const STableRow<FSceneOutlinerTreeItemPtr>& InRow)
+{
+	return SNew(SVariablesOutlinerStructSharedVariablesLabel, *this, Outliner, InRow);
+}
+
+FString FVariablesOutlinerStructSharedVariablesItem::GetPackageName() const
+{
+	const UAnimNextRigVMAssetEntry* Entry = WeakEntry.Get();
+	if(Entry == nullptr)
+	{
+		return ISceneOutlinerTreeItem::GetPackageName();
+	}
+
+	return Entry->GetPackage()->GetName();
+}
+
+}
+
+#undef LOCTEXT_NAMESPACE // "VariablesOutlinerTreeItem"
\ No newline at end of file
diff --git a/Engine/Plugins/Experimental/AnimNext/Source/AnimNextEditor/Private/Variables/Outliner/VariablesOutlinerStructSharedVariablesItem.h b/Engine/Plugins/Experimental/AnimNext/Source/AnimNextEditor/Private/Variables/Outliner/VariablesOutlinerStructSharedVariablesItem.h
new file mode 100644
--- /dev/null
+++ b/Engine/Plugins/Experimental/AnimNext/Source/AnimNextEditor/Private/Variables/Outliner/VariablesOutlinerStructSharedVariablesItem.h
@@ -0,0 +1,33 @@
+﻿// Copyright Epic Games, Inc. All Rights Reserved.
+
+#pragma once
+
+#include "ISceneOutlinerTreeItem.h"
+
+class UAnimNextSharedVariablesEntry;
+
+namespace UE::AnimNext::Editor
+{
+
+class IWorkspaceOutlinerItemDetails;
+
+struct FVariablesOutlinerStructSharedVariablesItem : ISceneOutlinerTreeItem
+{
+	static const FSceneOutlinerTreeItemType Type;
+
+	FVariablesOutlinerStructSharedVariablesItem(UAnimNextSharedVariablesEntry* InEntry);
+
+	// Begin ISceneOutlinerTreeItem overrides
+	virtual bool IsValid() const override;
+	virtual FSceneOutlinerTreeItemID GetID() const override;
+	virtual FString GetDisplayString() const override;
+	virtual bool CanInteract() const override { return true; }
+	virtual TSharedRef<SWidget> GenerateLabelWidget(ISceneOutliner& Outliner, const STableRow<FSceneOutlinerTreeItemPtr>& InRow) override;
+	virtual FString GetPackageName() const override;
+	// End ISceneOutlinerTreeItem overrides
+
+	// Weak-ptr to the represented entry
+	TWeakObjectPtr<UAnimNextSharedVariablesEntry> WeakEntry;
+};
+
+}
diff --git a/Engine/Plugins/Experimental/AnimNext/Source/AnimNextEditor/Private/Variables/SAddVariablesDialog.cpp b/Engine/Plugins/Experimental/AnimNext/Source/AnimNextEditor/Private/Variables/SAddVariablesDialog.cpp
--- a/Engine/Plugins/Experimental/AnimNext/Source/AnimNextEditor/Private/Variables/SAddVariablesDialog.cpp
+++ b/Engine/Plugins/Experimental/AnimNext/Source/AnimNextEditor/Private/Variables/SAddVariablesDialog.cpp
@@ -192,7 +192,7 @@ void SAddVariablesDialog::Construct(const FArguments& InArgs, const TArray<UAnim
 							{
 								FAnimNextAssetRegistryExports Exports;
 								UncookedOnly::FUtils::GetExportedVariablesForAsset(InAssetData, Exports);
-								if(Exports.Variables.Num() == 0)
+								if(Exports.Exports.Num() == 0 && InAssetData.GetClass(EResolveClass::Yes) != UAnimNextSharedVariables::StaticClass())
 								{
 									return true;
 								}
@@ -236,7 +236,7 @@ void SAddVariablesDialog::Construct(const FArguments& InArgs, const TArray<UAnim
 									return true;
 								}
 
-								virtual bool IsUnloadedStructAllowed(const FStructViewerInitializationOptions& InInitOptions, const FSoftObjectPath& InStructPath, TSharedRef<class FStructViewerFilterFuncs> InFilterFuncs)
+								virtual bool IsUnloadedStructAllowed(const FStructViewerInitializationOptions& InInitOptions, const FSoftObjectPath& InStructPath, TSharedRef<class FStructViewerFilterFuncs> InFilterFuncs) override
 								{
 									return false;
 								};
@@ -370,70 +370,11 @@ void SAddVariablesDialog::DeleteSelectedItems()
 	RefreshEntries();
 }
 
-static FName GetNewVariableName(FName InBaseName, TArrayView<FName> InExistingNames)
-{
-	auto NameExists = [&InExistingNames](FName InName)
-	{
-		for(FName AdditionalName : InExistingNames)
-		{
-			if(AdditionalName == InName)
-			{
-				return true;
-			}
-		}
-
-		return false;
-	};
-
-	if(!NameExists(InBaseName))
-	{
-		// Early out - name is valid
-		return InBaseName;
-	}
-
-	int32 PostFixIndex = 0;
-	TStringBuilder<128> StringBuilder;
-	while(true)
-	{
-		StringBuilder.Reset();
-		InBaseName.GetDisplayNameEntry()->AppendNameToString(StringBuilder);
-		StringBuilder.Appendf(TEXT("_%d"), PostFixIndex++);
-
-		FName TestName(StringBuilder.ToString()); 
-		if(!NameExists(TestName))
-		{
-			return TestName;
-		}
-	}
-}
-
-void SAddVariablesDialog::GetPendingNamesRecursive(UAnimNextRigVMAssetEditorData* InEditorData, TArray<FName>& OutPendingNames) const
-{
-	for(UAnimNextRigVMAssetEntry* Entry : InEditorData->Entries)
-	{
-		if(UAnimNextVariableEntry* VariableEntry = Cast<UAnimNextVariableEntry>(Entry))
-		{
-			if(VariableEntry->GetExportAccessSpecifier() == EAnimNextExportAccessSpecifier::Public)
-			{
-				OutPendingNames.Add(Entry->GetEntryName());
-			}
-		}
-		else if(UAnimNextSharedVariablesEntry* DataInterfaceEntry = Cast<UAnimNextSharedVariablesEntry>(Entry))
-		{
-			if(DataInterfaceEntry->GetAsset())
-			{
-				UAnimNextSharedVariables_EditorData* EditorData = UE::AnimNext::UncookedOnly::FUtils::GetEditorData<UAnimNextSharedVariables_EditorData>(DataInterfaceEntry->GetAsset());
-				GetPendingNamesRecursive(EditorData, OutPendingNames);
-			}
-		}
-	}
-}
-
 void SAddVariablesDialog::GetPendingNames(TArray<FName>& OutPendingNames) const
 {
 	for(UAnimNextRigVMAssetEditorData* EditorData : AssetEditorDatas)
 	{
-		GetPendingNamesRecursive(EditorData, OutPendingNames);
+		UncookedOnly::FUtils::GetVariableNames(EditorData, OutPendingNames);
 	}
 }
 
@@ -462,7 +403,7 @@ void SAddVariablesDialog::AddEntry(const FAnimNextParamType& InParamType)
 
 	GetPendingNames(PendingNames);
 
-	FName VariableName = GetNewVariableName(Settings->GetLastVariableName(), PendingNames);
+	FName VariableName = UE::AnimNext::UncookedOnly::FUtils::GetValidVariableName(Settings->GetLastVariableName(), PendingNames);
 	RootEntries.Add(MakeShared<FVariableToAddEntry>(InParamType.IsValid() ? InParamType : Settings->GetLastVariableType(), VariableName, SharedThis(this)));
 
 	RefreshEntries();
@@ -472,7 +413,7 @@ void SAddVariablesDialog::AddSharedVariables(UAnimNextSharedVariables* InSharedV
 {
 	TSharedRef<FSharedVariablesToAddEntry> NewEntry = MakeShared<FSharedVariablesToAddEntry>(InSharedVariables, SharedThis(this));
 
-	auto AddVariable = [this, &NewEntry](UAnimNextVariableEntry* InVariableEntry)
+	auto AddVariable = [this, &NewEntry](const UAnimNextVariableEntry* InVariableEntry)
 	{
 		if(InVariableEntry->GetExportAccessSpecifier() == EAnimNextExportAccessSpecifier::Public)
 		{
@@ -517,7 +458,7 @@ void SAddVariablesDialog::AddSharedVariablesStruct(const UScriptStruct* InStruct
 {
 	TSharedRef<FSharedVariablesToAddEntry> NewEntry = MakeShared<FSharedVariablesToAddEntry>(InStruct, SharedThis(this));
 
-	auto AddVariable = [this, &NewEntry](FProperty* InProperty)
+	auto AddVariable = [this, &NewEntry](const FProperty* InProperty)
 	{
 		TSharedRef<FVariableToAddEntry> NewSubEntry = MakeShared<FVariableToAddEntry>(FAnimNextParamType::FromProperty(InProperty), InProperty->GetFName(), SharedThis(this));
 		NewSubEntry->Parent = NewEntry;
@@ -602,7 +543,7 @@ class SVariableToAdd : public SMultiColumnTableRow<TSharedRef<SAddVariablesDialo
 				.Font(IDetailLayoutBuilder::GetDetailFont())
 				.IsReadOnly_Lambda([this]()
 				{
-					// Cant rename entries from data interfaces (i.e. with parents)
+					// Cant rename entries from shared variables (i.e. with parents)
 					check(Entry->EntryType == SAddVariablesDialog::EEntryType::Variable);
 					return Entry->Parent.IsValid();
 				})
@@ -768,7 +709,7 @@ class SVariableToAdd : public SMultiColumnTableRow<TSharedRef<SAddVariablesDialo
 						SNew(SPinTypeSelector, FGetPinTypeTree::CreateLambda(GetFilteredVariableTypeTree))
 							.IsEnabled_Lambda([this]()
 							{
-								// Cant rename entries from data interfaces (i.e. with parents)
+								// Cant rename entries from shared variables (i.e. with parents)
 								check(Entry->EntryType == SAddVariablesDialog::EEntryType::Variable);
 								return !Entry->Parent.IsValid();
 							})
diff --git a/Engine/Plugins/Experimental/AnimNext/Source/AnimNextEditor/Private/Variables/SAddVariablesDialog.h b/Engine/Plugins/Experimental/AnimNext/Source/AnimNextEditor/Private/Variables/SAddVariablesDialog.h
--- a/Engine/Plugins/Experimental/AnimNext/Source/AnimNextEditor/Private/Variables/SAddVariablesDialog.h
+++ b/Engine/Plugins/Experimental/AnimNext/Source/AnimNextEditor/Private/Variables/SAddVariablesDialog.h
@@ -150,8 +150,6 @@ private:
 	
 	TSharedRef<SWidget> HandleGetAddVariableMenuContent(TSharedPtr<FEntry> InEntry);
 
-	void GetPendingNamesRecursive(UAnimNextRigVMAssetEditorData* InEditorData, TArray<FName>& OutPendingNames) const;
-
 	void GetPendingNames(TArray<FName>& OutPendingNames) const;
 
 private:
diff --git a/Engine/Plugins/Experimental/AnimNext/Source/AnimNextEditor/Private/Variables/SVariablePicker.cpp b/Engine/Plugins/Experimental/AnimNext/Source/AnimNextEditor/Private/Variables/SVariablePicker.cpp
--- a/Engine/Plugins/Experimental/AnimNext/Source/AnimNextEditor/Private/Variables/SVariablePicker.cpp
+++ b/Engine/Plugins/Experimental/AnimNext/Source/AnimNextEditor/Private/Variables/SVariablePicker.cpp
@@ -9,11 +9,9 @@
 #include "Widgets/Layout/SBox.h"
 #include "Widgets/Images/SImage.h"
 #include "Framework/Application/SlateApplication.h"
-#include "AnimNextUncookedOnly/Private/AnimNextUncookedOnlyModule.h"
 #include "Widgets/PropertyViewer/SPropertyViewer.h"
 #include "Variables/AnimNextSoftVariableReference.h"
 #include "Widgets/Input/SCheckBox.h"
-#include "Modules/ModuleManager.h"
 
 #define LOCTEXT_NAMESPACE "SVariablePicker"
 
@@ -86,30 +84,32 @@ void SVariablePicker::RefreshEntries()
 	CachedContainers.Reset();
 	ContainerMap.Reset();
 
-	UncookedOnly::IAnimNextUncookedOnlyModule& Module = FModuleManager::GetModuleChecked<UncookedOnly::IAnimNextUncookedOnlyModule>("AnimNextUncookedOnly");
-
 	// Add variables exposed on assets
 	TMap<FAssetData, FAnimNextAssetRegistryExports> Exports;
 	if(UncookedOnly::FUtils::GetExportedVariablesFromAssetRegistry(Exports))
 	{
 		for(const TPair<FAssetData, FAnimNextAssetRegistryExports>& ExportPair : Exports)
 		{
-			if(ExportPair.Value.Variables.Num() > 0)
+			if(ExportPair.Value.Exports.Num() > 0)
 			{
 				// Add a placeholder struct for this asset's properties
 				TArray<FPropertyBagPropertyDesc> PropertyDescs;
-				PropertyDescs.Reserve(ExportPair.Value.Variables.Num());
-				for(const FAnimNextAssetRegistryExportedVariable& VariableEntry : ExportPair.Value.Variables)
+				PropertyDescs.Reserve(ExportPair.Value.Exports.Num());
+				for(const FAnimNextExport& Export : ExportPair.Value.Exports)
 				{
-					if(EnumHasAllFlags(VariableEntry.GetFlags(), EAnimNextExportedVariableFlags::Declared | EAnimNextExportedVariableFlags::Public) && VariableEntry.Name != NAME_None && VariableEntry.Type.IsValid())
+					if (Export.Data.GetScriptStruct() == FAnimNextVariableDeclarationData::StaticStruct())
 					{
-						FAnimNextSoftVariableReference VariableReference(VariableEntry.Name, ExportPair.Key.ToSoftObjectPath());
-
-						if (!Args.OnFilterVariable.IsBound() || Args.OnFilterVariable.Execute(VariableReference) == EFilterVariableResult::Include)
+						const FAnimNextVariableDeclarationData& VariableEntry = Export.Data.Get<FAnimNextVariableDeclarationData>();
+						if(EnumHasAllFlags(static_cast<EAnimNextExportedVariableFlags>(VariableEntry.Flags), Args.FlagFilter) && Export.Identifier != NAME_None && VariableEntry.Type.IsValid())
 						{
-							if (!Args.OnFilterVariableType.IsBound() || Args.OnFilterVariableType.Execute(VariableEntry.Type) == EFilterVariableResult::Include)
+							FAnimNextSoftVariableReference VariableReference(Export.Identifier, ExportPair.Key.ToSoftObjectPath());
+
+							if (!Args.OnFilterVariable.IsBound() || Args.OnFilterVariable.Execute(VariableReference) == EFilterVariableResult::Include)
 							{
-								PropertyDescs.Emplace(VariableEntry.Name, VariableEntry.Type.GetContainerType(), VariableEntry.Type.GetValueType(), VariableEntry.Type.GetValueTypeObject());
+								if (!Args.OnFilterVariableType.IsBound() || Args.OnFilterVariableType.Execute(VariableEntry.Type) == EFilterVariableResult::Include)
+								{
+									PropertyDescs.Emplace(Export.Identifier, VariableEntry.Type.GetContainerType(), VariableEntry.Type.GetValueType(), VariableEntry.Type.GetValueTypeObject());
+								}
 							}
 						}
 					}
diff --git a/Engine/Plugins/Experimental/AnimNext/Source/AnimNextEditor/Private/Variables/SVariablesView.cpp b/Engine/Plugins/Experimental/AnimNext/Source/AnimNextEditor/Private/Variables/SVariablesView.cpp
--- a/Engine/Plugins/Experimental/AnimNext/Source/AnimNextEditor/Private/Variables/SVariablesView.cpp
+++ b/Engine/Plugins/Experimental/AnimNext/Source/AnimNextEditor/Private/Variables/SVariablesView.cpp
@@ -12,16 +12,15 @@
 #include "ScopedTransaction.h"
 #include "SPositiveActionButton.h"
 #include "SSceneOutliner.h"
-#include "Framework/Commands/GenericCommands.h"
 #include "Outliner/VariablesOutlinerColumns.h"
 #include "Outliner/VariablesOutlinerMode.h"
 #include "Outliner/VariablesOutlinerEntryItem.h"
 #include "Variables/SAddVariablesDialog.h"
-#include "AnimNextRigVMAsset.h"
 #include "IWorkspaceEditor.h"
+#include "IWorkspaceEditorModule.h"
 #include "WorkspaceAssetRegistryInfo.h"
-#include "Common/AnimNextFunctionItemDetails.h"
 #include "Outliner/VariablesOutlinerAssetItem.h"
+#include "Entries/AnimNextSharedVariablesEntry.h"
 
 #define LOCTEXT_NAMESPACE "SVariablesView"
 
@@ -157,113 +156,101 @@ FReply SVariablesOutliner::HandleAddVariablesClicked()
 	return FReply::Handled();
 }
 
-void SVariablesOutliner::SetAssets(TConstArrayView<TSoftObjectPtr<UAnimNextRigVMAsset>> InAssets)
-{
-	for(const TSoftObjectPtr<UAnimNextRigVMAsset>& CurrentSoftAsset : Assets)
+void SVariablesOutliner::RegisterAssetDelegates(const UAnimNextRigVMAsset* InAsset)
+{	
+	// Bind for any modification callbacks
+	UAnimNextRigVMAssetEditorData* EditorData = UncookedOnly::FUtils::GetEditorData<UAnimNextRigVMAssetEditorData>(InAsset);
+	if(EditorData != nullptr)
 	{
-		if(UAnimNextRigVMAsset* CurrentAsset = CurrentSoftAsset.Get())
+		EditorData->ModifiedDelegate.AddSP(this, &SVariablesOutliner::OnEditorDataModified);
+
+		for (const TObjectPtr<UAnimNextRigVMAssetEntry>& Entry : EditorData->Entries)
 		{
-			UAnimNextRigVMAssetEditorData* EditorData = UncookedOnly::FUtils::GetEditorData<UAnimNextRigVMAssetEditorData>(CurrentAsset);
-			if(EditorData != nullptr)
+			if (UAnimNextSharedVariablesEntry* SharedVariablesEntry = Cast<UAnimNextSharedVariablesEntry>(Entry.Get()))
 			{
-				EditorData->ModifiedDelegate.RemoveAll(this);
+				if (TObjectPtr<const UAnimNextSharedVariables> SharedVariablesAsset = SharedVariablesEntry->GetAsset())
+				{
+					RegisterAssetDelegates(SharedVariablesAsset);
+					Assets.AddUnique(Cast<UAnimNextRigVMAsset>(const_cast<UAnimNextSharedVariables*>(SharedVariablesAsset.Get())));					
+				}
 			}
 		}
 	}
+}
 
-	Assets = InAssets;
-
-	for(const TSoftObjectPtr<UAnimNextRigVMAsset>& NewSoftAsset : Assets)
+void SVariablesOutliner::UnregisterAssetDelegates(const UAnimNextRigVMAsset* InAsset)
+{
+	UAnimNextRigVMAssetEditorData* EditorData = UncookedOnly::FUtils::GetEditorData<UAnimNextRigVMAssetEditorData>(InAsset);
+	if(EditorData != nullptr)
 	{
-		if(UAnimNextRigVMAsset* NewAsset = NewSoftAsset.Get())
+		EditorData->ModifiedDelegate.RemoveAll(this);
+
+		for (const TObjectPtr<UAnimNextRigVMAssetEntry>& Entry : EditorData->Entries)
 		{
-			UAnimNextRigVMAssetEditorData* EditorData = UncookedOnly::FUtils::GetEditorData<UAnimNextRigVMAssetEditorData>(NewAsset);
-			if(EditorData != nullptr)
+			if (UAnimNextSharedVariablesEntry* SharedVariablesEntry = Cast<UAnimNextSharedVariablesEntry>(Entry.Get()))
 			{
-				EditorData->ModifiedDelegate.AddSP(this, &SVariablesOutliner::OnEditorDataModified);
+				if (TObjectPtr<const UAnimNextSharedVariables> SharedVariablesAsset = SharedVariablesEntry->GetAsset())
+				{
+					UnregisterAssetDelegates(SharedVariablesAsset);					
+				}
 			}
 		}
 	}
-
-	FullRefresh();
 }
 
-void SVariablesOutliner::HandleAssetLoaded(const FSoftObjectPath& InSoftObjectPath, UAnimNextRigVMAsset* InAsset)
+void SVariablesOutliner::SetAssets(TConstArrayView<TSoftObjectPtr<UAnimNextRigVMAsset>> InAssets)
 {
-	if(Assets.Contains(TSoftObjectPtr<UAnimNextRigVMAsset>(InAsset)))
+	for(const TSoftObjectPtr<UAnimNextRigVMAsset>& CurrentSoftAsset : Assets)
 	{
-		// Bind for any modification callbacks
-		UAnimNextRigVMAssetEditorData* EditorData = UncookedOnly::FUtils::GetEditorData<UAnimNextRigVMAssetEditorData>(InAsset);
-		if(EditorData != nullptr)
+		if(UAnimNextRigVMAsset* CurrentAsset = CurrentSoftAsset.Get())
 		{
-			EditorData->ModifiedDelegate.AddSP(this, &SVariablesOutliner::OnEditorDataModified);
+			UnregisterAssetDelegates(CurrentAsset);
 		}
-
-		FullRefresh();
 	}
-}
 
-void SVariablesOutliner::OnEditorDataModified(UAnimNextRigVMAssetEditorData* InEditorData, EAnimNextEditorDataNotifType InType, UObject* InSubject)
-{
-	ensure(Assets.Contains(TSoftObjectPtr<UAnimNextRigVMAsset>(UncookedOnly::FUtils::GetAsset<UAnimNextRigVMAsset>(InEditorData))));
+	Assets = InAssets;
 
-	switch(InType)
+	for(const TSoftObjectPtr<UAnimNextRigVMAsset>& NewSoftAsset : InAssets)
 	{
-	case EAnimNextEditorDataNotifType::UndoRedo:
-	case EAnimNextEditorDataNotifType::EntryAdded:
-	case EAnimNextEditorDataNotifType::EntryRemoved:
-		FullRefresh();
-		break;
-	default:
-		break;
+		if(UAnimNextRigVMAsset* NewAsset = NewSoftAsset.Get())
+		{
+			RegisterAssetDelegates(NewAsset);
+		}
 	}
-}
 
-bool SVariablesOutliner::IsEnabled() const
-{
-	return Assets.Num() > 0;
+	FullRefresh();
 }
 
-void SVariablesView::Construct(const FArguments& InArgs, TSharedRef<UE::Workspace::IWorkspaceEditor> InWorkspaceEditor)
+void SVariablesOutliner::SetExport(const FWorkspaceOutlinerItemExport& InExport)
 {
-	InWorkspaceEditor->OnOutlinerSelectionChanged().AddSP(this, &SVariablesView::HandleWorkspaceOutlinerSelectionChanged);
-
-	FSceneOutlinerInitializationOptions InitOptions;
-	InitOptions.OutlinerIdentifier = TEXT("AnimNextVariablesOutliner");
-	InitOptions.bShowHeaderRow = true;
-	InitOptions.ColumnMap.Add(FSceneOutlinerBuiltInColumnTypes::Label(), FSceneOutlinerColumnInfo(ESceneOutlinerColumnVisibility::Visible, 0, FCreateSceneOutlinerColumn(), false, 0.5f));
-	InitOptions.ColumnMap.Add(FVariablesOutlinerTypeColumn::GetID(), FSceneOutlinerColumnInfo(ESceneOutlinerColumnVisibility::Visible, 10, FCreateSceneOutlinerColumn::CreateLambda([](ISceneOutliner& InSceneOutliner) { return MakeShared<FVariablesOutlinerTypeColumn>(InSceneOutliner); }), false));
-	InitOptions.ColumnMap.Add(FVariablesOutlinerValueColumn::GetID(), FSceneOutlinerColumnInfo(ESceneOutlinerColumnVisibility::Visible, 20, FCreateSceneOutlinerColumn::CreateLambda([](ISceneOutliner& InSceneOutliner) { return MakeShared<FVariablesOutlinerValueColumn>(InSceneOutliner); }), true, 0.5f));
-	InitOptions.ColumnMap.Add(FSceneOutlinerSourceControlColumn::GetID(), FSceneOutlinerColumnInfo(ESceneOutlinerColumnVisibility::Visible, 30, FCreateSceneOutlinerColumn::CreateLambda([](ISceneOutliner& InSceneOutliner) { return MakeShared<FSceneOutlinerSourceControlColumn>(InSceneOutliner); }), true));
-	InitOptions.ColumnMap.Add(FVariablesOutlinerAccessSpecifierColumn::GetID(), FSceneOutlinerColumnInfo(ESceneOutlinerColumnVisibility::Visible, 40, FCreateSceneOutlinerColumn::CreateLambda([](ISceneOutliner& InSceneOutliner) { return MakeShared<FVariablesOutlinerAccessSpecifierColumn>(InSceneOutliner); }), false));
-	InitOptions.ModeFactory = FCreateSceneOutlinerMode::CreateLambda([this, WeakWorkspaceEditor=InWorkspaceEditor.ToWeakPtr()](SSceneOutliner* InOutliner) { return new FVariablesOutlinerMode(static_cast<SVariablesOutliner*>(InOutliner), WeakWorkspaceEditor.Pin().ToSharedRef()); });
-
-	VariablesOutliner = SNew(SVariablesOutliner, InitOptions);
-	VariablesOutliner->SetEnabled(MakeAttributeSP(VariablesOutliner.Get(), &SVariablesOutliner::IsEnabled));
+	Export = InExport;
 
-	TArray<FWorkspaceOutlinerItemExport> SelectedExports;
-	InWorkspaceEditor->GetOutlinerSelection(SelectedExports);
-	HandleWorkspaceOutlinerSelectionChanged(SelectedExports);
-
-	ChildSlot
-	[
-		VariablesOutliner.ToSharedRef()
-	];
-}
-
-void SVariablesView::HandleWorkspaceOutlinerSelectionChanged(TConstArrayView<FWorkspaceOutlinerItemExport> InExports)
-{
-	TArray<TSoftObjectPtr<UAnimNextRigVMAsset>> Assets;
 	TArray<FSoftObjectPath> AssetsToAsyncLoad;
-	for(const FWorkspaceOutlinerItemExport& Export : InExports)
+	TArray<TSoftObjectPtr<UAnimNextRigVMAsset>> ExportAssets;
+	if (FWorkspaceOutlinerAssetReferenceItemData::IsAssetReference(Export))
+	{
+		// Reference item
+		TArray<FSoftObjectPath> AssetPaths;
+		Export.GetAssetPaths(AssetPaths);
+
+		for (const FSoftObjectPath& Path : AssetPaths)
+		{
+			ExportAssets.Add(TSoftObjectPtr<UAnimNextRigVMAsset>(Path));
+			
+			if(Path.ResolveObject() == nullptr)
+			{
+				AssetsToAsyncLoad.Add(Path);
+			}
+		}
+	}
+	else
 	{
 		const FSoftObjectPath FirstAssetPath = Export.GetFirstAssetPath();
 		if(FirstAssetPath.IsAsset())
 		{
-			Assets.Add(TSoftObjectPtr<UAnimNextRigVMAsset>(FirstAssetPath));
+			ExportAssets.Add(TSoftObjectPtr<UAnimNextRigVMAsset>(FirstAssetPath));
 
 			UObject* ResolvedObject = FirstAssetPath.ResolveObject();
-			UAnimNextRigVMAsset* ResolvedAsset = Cast<UAnimNextRigVMAsset>(ResolvedObject);
 			if(ResolvedObject == nullptr)
 			{
 				AssetsToAsyncLoad.Add(FirstAssetPath);
@@ -271,15 +258,12 @@ void SVariablesView::HandleWorkspaceOutlinerSelectionChanged(TConstArrayView<FWo
 		}
 		else if(Export.HasData() && Export.GetData().GetScriptStruct()->IsChildOf(FAnimNextAssetEntryOutlinerData::StaticStruct()))
 		{
-			// TODO: As we are not showing references in the workspace yet, we are just traversing the outer here and showing the containing asset's
-			// variables. Eventually we would want to recurse up to the root of the export hierarchy, which will entail changing the signature of
-			// FOnOutlinerSelectionChanged to include parent items somehow, or adding parent query functionality to the workspace API
 			const FAnimNextAssetEntryOutlinerData& EntryData = Export.GetData().Get<FAnimNextAssetEntryOutlinerData>();
 			if(EntryData.SoftEntryPtr.IsValid())
 			{
 				if(UAnimNextRigVMAsset* Asset = EntryData.GetEntry()->GetTypedOuter<UAnimNextRigVMAsset>())
 				{
-					Assets.Add(TSoftObjectPtr<UAnimNextRigVMAsset>(Asset));
+					ExportAssets.Add(TSoftObjectPtr<UAnimNextRigVMAsset>(Asset));
 				}
 			}
 		}
@@ -292,19 +276,21 @@ void SVariablesView::HandleWorkspaceOutlinerSelectionChanged(TConstArrayView<FWo
 				{
 					if(UAnimNextRigVMAsset* Asset = EdGraph->GetTypedOuter<UAnimNextRigVMAsset>())
 					{
-						Assets.Add(TSoftObjectPtr<UAnimNextRigVMAsset>(Asset));
+						ExportAssets.Add(TSoftObjectPtr<UAnimNextRigVMAsset>(Asset));
 					}
 				}
 			}
 		}
 	}
 
-	VariablesOutliner->SetAssets(Assets);
+	SetAssets(ExportAssets);
 
 	// Try async load any missing assets
 	for(const FSoftObjectPath& AssetPath : AssetsToAsyncLoad)
 	{
-		AssetPath.LoadAsync(FLoadSoftObjectPathAsyncDelegate::CreateLambda([WeakVariablesOutliner = TWeakPtr<SVariablesOutliner>(VariablesOutliner)](const FSoftObjectPath& InSoftObjectPath, UObject* InObject)
+		TWeakPtr<SVariablesOutliner> WeakOutliner = StaticCastSharedRef<SVariablesOutliner>(AsShared());
+		
+		AssetPath.LoadAsync(FLoadSoftObjectPathAsyncDelegate::CreateLambda([WeakOutliner](const FSoftObjectPath& InSoftObjectPath, UObject* InObject)
 		{
 			UAnimNextRigVMAsset* Asset = Cast<UAnimNextRigVMAsset>(InObject);
 			if(Asset == nullptr)
@@ -312,7 +298,7 @@ void SVariablesView::HandleWorkspaceOutlinerSelectionChanged(TConstArrayView<FWo
 				return;
 			}
 
-			TSharedPtr<SVariablesOutliner> PinnedVariablesOutliner = WeakVariablesOutliner.Pin();
+			TSharedPtr<SVariablesOutliner> PinnedVariablesOutliner = WeakOutliner.Pin();
 			if(!PinnedVariablesOutliner.IsValid())
 			{
 				return;
@@ -323,6 +309,89 @@ void SVariablesView::HandleWorkspaceOutlinerSelectionChanged(TConstArrayView<FWo
 	}
 }
 
+void SVariablesOutliner::HandleAssetLoaded(const FSoftObjectPath& InSoftObjectPath, UAnimNextRigVMAsset* InAsset)
+{	
+	if(InAsset && Assets.Contains(TSoftObjectPtr<UAnimNextRigVMAsset>(InAsset)))
+	{
+		RegisterAssetDelegates(InAsset);
+
+		FullRefresh();
+	}
+}
+
+void SVariablesOutliner::OnEditorDataModified(UAnimNextRigVMAssetEditorData* InEditorData, EAnimNextEditorDataNotifType InType, UObject* InSubject)
+{	
+	ensure(Assets.Contains(TSoftObjectPtr<UAnimNextRigVMAsset>(UncookedOnly::FUtils::GetAsset<UAnimNextRigVMAsset>(InEditorData))));
+
+	switch(InType)
+	{
+	case EAnimNextEditorDataNotifType::UndoRedo:
+	case EAnimNextEditorDataNotifType::EntryAdded:
+	case EAnimNextEditorDataNotifType::EntryRemoved:		
+	case EAnimNextEditorDataNotifType::VariableCategoryChanged:
+	case EAnimNextEditorDataNotifType::CategoryAdded:
+	case EAnimNextEditorDataNotifType::CategoryChanged:
+		FullRefresh();
+		break;
+	default:
+		break;
+	}
+}
+
+void SVariablesOutliner::SetHighlightedItem(FSceneOutlinerTreeItemPtr Item) const
+{
+	GetTreeView()->Private_SetItemHighlighted(Item, true);
+}
+
+void SVariablesOutliner::ClearHighlightedItem(FSceneOutlinerTreeItemPtr Item) const
+{
+	GetTreeView()->Private_SetItemHighlighted(Item, false);
+}
+
+bool SVariablesOutliner::HasAssets() const
+{
+	return Assets.Num() > 0;
+}
+
+void SVariablesView::Construct(const FArguments& InArgs, TSharedRef<UE::Workspace::IWorkspaceEditor> InWorkspaceEditor)
+{
+	FSceneOutlinerInitializationOptions InitOptions;
+	InitOptions.OutlinerIdentifier = TEXT("AnimNextVariablesOutliner");
+	InitOptions.bShowHeaderRow = true;
+	InitOptions.ColumnMap.Add(FSceneOutlinerBuiltInColumnTypes::Label(), FSceneOutlinerColumnInfo(ESceneOutlinerColumnVisibility::Visible, 0, FCreateSceneOutlinerColumn(), false, 0.5f));
+	InitOptions.ColumnMap.Add(FVariablesOutlinerTypeColumn::GetID(), FSceneOutlinerColumnInfo(ESceneOutlinerColumnVisibility::Visible, 10, FCreateSceneOutlinerColumn::CreateLambda([](ISceneOutliner& InSceneOutliner) { return MakeShared<FVariablesOutlinerTypeColumn>(InSceneOutliner); }), false));
+	InitOptions.ColumnMap.Add(FVariablesOutlinerValueColumn::GetID(), FSceneOutlinerColumnInfo(ESceneOutlinerColumnVisibility::Visible, 20, FCreateSceneOutlinerColumn::CreateLambda([](ISceneOutliner& InSceneOutliner) { return MakeShared<FVariablesOutlinerValueColumn>(InSceneOutliner); }), true, 0.5f));
+	InitOptions.ColumnMap.Add(FSceneOutlinerSourceControlColumn::GetID(), FSceneOutlinerColumnInfo(ESceneOutlinerColumnVisibility::Visible, 30, FCreateSceneOutlinerColumn::CreateLambda([](ISceneOutliner& InSceneOutliner) { return MakeShared<FSceneOutlinerSourceControlColumn>(InSceneOutliner); }), true));
+	InitOptions.ColumnMap.Add(FVariablesOutlinerAccessSpecifierColumn::GetID(), FSceneOutlinerColumnInfo(ESceneOutlinerColumnVisibility::Visible, 40, FCreateSceneOutlinerColumn::CreateLambda([](ISceneOutliner& InSceneOutliner) { return MakeShared<FVariablesOutlinerAccessSpecifierColumn>(InSceneOutliner); }), false));
+	InitOptions.ModeFactory = FCreateSceneOutlinerMode::CreateLambda([this, WeakWorkspaceEditor=InWorkspaceEditor.ToWeakPtr()](SSceneOutliner* InOutliner) { return new FVariablesOutlinerMode(static_cast<SVariablesOutliner*>(InOutliner), WeakWorkspaceEditor.Pin().ToSharedRef()); });
+
+	VariablesOutliner = SNew(SVariablesOutliner, InitOptions);
+	VariablesOutliner->SetEnabled(MakeAttributeSP(VariablesOutliner.Get(), &SVariablesOutliner::HasAssets));
+
+	InWorkspaceEditor->OnFocusedDocumentChanged().AddSP(this, &SVariablesView::HandleFocusedDocumentChanged);
+	const Workspace::FWorkspaceDocument& Document = InWorkspaceEditor->GetFocusedWorkspaceDocument();
+	HandleFocusedDocumentChanged(Document);
+
+	ChildSlot
+	[
+		VariablesOutliner.ToSharedRef()
+	];
+}
+
+void SVariablesView::SetExportDirectly(const FWorkspaceOutlinerItemExport& InExport) const
+{
+	if (VariablesOutliner.IsValid())
+	{
+		VariablesOutliner->SetExport(InExport);
+	}
+}
+
+void SVariablesView::HandleFocusedDocumentChanged(const UE::Workspace::FWorkspaceDocument& InDocument) const
+{
+	const FWorkspaceOutlinerItemExport& Export = InDocument.Export;
+	VariablesOutliner->SetExport(Export);
+}
+
 FAnimNextVariablesTabSummoner::FAnimNextVariablesTabSummoner(TSharedPtr<UE::Workspace::IWorkspaceEditor> InHostingApp)
 	: FWorkflowTabFactory(VariablesTabName, StaticCastSharedPtr<FAssetEditorToolkit>(InHostingApp))
 {
@@ -334,11 +403,8 @@ FAnimNextVariablesTabSummoner::FAnimNextVariablesTabSummoner(TSharedPtr<UE::Work
 
 	VariablesView = SNew(SVariablesView, InHostingApp.ToSharedRef());
 
-	TArray<FWorkspaceOutlinerItemExport> Exports;
-	if(InHostingApp->GetOutlinerSelection(Exports))
-	{
-		VariablesView->HandleWorkspaceOutlinerSelectionChanged(Exports);
-	}
+	const Workspace::FWorkspaceDocument& Document = InHostingApp->GetFocusedWorkspaceDocument();
+	VariablesView->VariablesOutliner->SetExport(Document.Export);
 }
 
 TSharedRef<SWidget> FAnimNextVariablesTabSummoner::CreateTabBody(const FWorkflowTabSpawnInfo& Info) const
diff --git a/Engine/Plugins/Experimental/AnimNext/Source/AnimNextEditor/Private/Variables/SVariablesView.h b/Engine/Plugins/Experimental/AnimNext/Source/AnimNextEditor/Private/Variables/SVariablesView.h
--- a/Engine/Plugins/Experimental/AnimNext/Source/AnimNextEditor/Private/Variables/SVariablesView.h
+++ b/Engine/Plugins/Experimental/AnimNext/Source/AnimNextEditor/Private/Variables/SVariablesView.h
@@ -4,6 +4,8 @@
 
 #include "ContentBrowserDelegates.h"
 #include "SSceneOutliner.h"
+#include "WorkspaceAssetRegistryInfo.h"
+#include "Editor/Persona/Private/SAnimAssetFindReplace.h"
 #include "Widgets/SCompoundWidget.h"
 #include "WorkflowOrientedApp/WorkflowTabFactory.h"
 
@@ -16,6 +18,7 @@ class UAnimNextRigVMAssetEditorData;
 namespace UE::Workspace
 {
 	class IWorkspaceEditor;
+	struct FWorkspaceDocument;
 }
 
 namespace UE::AnimNext::Editor
@@ -37,19 +40,30 @@ public:
 
 	FReply HandleAddVariablesClicked();
 
-	void SetAssets(TConstArrayView<TSoftObjectPtr<UAnimNextRigVMAsset>> InAssets);
+	void SetExport(const FWorkspaceOutlinerItemExport& InExport);
 
 	void OnEditorDataModified(UAnimNextRigVMAssetEditorData* InEditorData, EAnimNextEditorDataNotifType InType, UObject* InSubject);
 
+	void SetHighlightedItem(FSceneOutlinerTreeItemPtr Item) const;
+	void ClearHighlightedItem(FSceneOutlinerTreeItemPtr Item) const; 
+	
+	bool HasAssets() const;
+protected:
+	
+	void SetAssets(TConstArrayView<TSoftObjectPtr<UAnimNextRigVMAsset>> InAssets);
+
 	void HandleAssetLoaded(const FSoftObjectPath& InSoftObjectPath, UAnimNextRigVMAsset* InAsset);
 	
-	bool IsEnabled() const;
+	void RegisterAssetDelegates(const UAnimNextRigVMAsset* InAsset);
+	void UnregisterAssetDelegates(const UAnimNextRigVMAsset* InAsset);
 private:
 	friend class FVariablesOutlinerMode;
 	friend class FVariablesOutlinerHierarchy;
 
 	TArray<TSoftObjectPtr<UAnimNextRigVMAsset>> Assets;
 	TSharedPtr<SPositiveActionButton> AddVariablesButton;
+
+	FWorkspaceOutlinerItemExport Export;
 };
 
 class SVariablesView : public SCompoundWidget
@@ -59,12 +73,12 @@ class SVariablesView : public SCompoundWidget
 
 	void Construct(const FArguments& InArgs, TSharedRef<UE::Workspace::IWorkspaceEditor> InWorkspaceEditor);
 
+	void SetExportDirectly(const FWorkspaceOutlinerItemExport& InExport) const;
 private:
-	void HandleWorkspaceOutlinerSelectionChanged(TConstArrayView<FWorkspaceOutlinerItemExport> InExports);
+	void HandleFocusedDocumentChanged(const UE::Workspace::FWorkspaceDocument& InDocument) const;
 
 private:
 	friend struct FAnimNextVariablesTabSummoner;
-	
 	TSharedPtr<SVariablesOutliner> VariablesOutliner;
 };
 
diff --git a/Engine/Plugins/Experimental/AnimNext/Source/AnimNextEditor/Private/Variables/VariableBindingPropertyCustomization.cpp b/Engine/Plugins/Experimental/AnimNext/Source/AnimNextEditor/Private/Variables/VariableBindingPropertyCustomization.cpp
--- a/Engine/Plugins/Experimental/AnimNext/Source/AnimNextEditor/Private/Variables/VariableBindingPropertyCustomization.cpp
+++ b/Engine/Plugins/Experimental/AnimNext/Source/AnimNextEditor/Private/Variables/VariableBindingPropertyCustomization.cpp
@@ -2,6 +2,7 @@
 
 #include "VariableBindingPropertyCustomization.h"
 
+#include "AnimNextRigVMAssetEditorData.h"
 #include "DetailWidgetRow.h"
 #include "EdGraphSchema_K2.h"
 #include "EditorUtils.h"
@@ -27,8 +28,8 @@ void FVariableBindingPropertyCustomization::CustomizeHeader(TSharedRef<IProperty
 	BindingDataHandle = PropertyHandle->GetChildHandle(GET_MEMBER_NAME_CHECKED(FAnimNextVariableBinding, BindingData));
 	check(BindingDataHandle.IsValid());
 
-	PropertyHandle->SetOnPropertyValueChanged(FSimpleDelegate::CreateRaw(this, &FVariableBindingPropertyCustomization::RequestRefresh));
-	PropertyHandle->SetOnChildPropertyValueChanged(FSimpleDelegate::CreateRaw(this, &FVariableBindingPropertyCustomization::RequestRefresh));
+	PropertyHandle->SetOnPropertyValueChanged(FSimpleDelegate::CreateRaw(this, &FVariableBindingPropertyCustomization::OnBindingChanged));
+	PropertyHandle->SetOnChildPropertyValueChanged(FSimpleDelegate::CreateRaw(this, &FVariableBindingPropertyCustomization::OnBindingChanged));
 
 	TOptional<FAnimNextParamType> CommonType;
 
@@ -48,6 +49,12 @@ void FVariableBindingPropertyCustomization::CustomizeHeader(TSharedRef<IProperty
 				CommonType = FAnimNextParamType();
 				break;
 			}
+			
+			if (UAnimNextRigVMAssetEditorData* EditorData = VariableEntry->GetTypedOuter<UAnimNextRigVMAssetEditorData>())
+			{
+				TWeakObjectPtr<UAnimNextVariableEntry> WeakEntry(VariableEntry);
+				EditorData->OnModified().AddSP(this, &FVariableBindingPropertyCustomization::OnEditorDataModified, WeakEntry);
+			}
 		}
 	}
 
@@ -118,7 +125,7 @@ void FVariableBindingPropertyCustomization::CustomizeHeader(TSharedRef<IProperty
 			SNew(SInstancedStructPicker, BindingDataHandle, InCustomizationUtils.GetPropertyUtilities())
 			.OnStructPicked_Lambda([this](const UScriptStruct* InStruct)
 			{
-				RequestRefresh();
+				OnBindingChanged();
 			})
 		]
 	];
@@ -132,6 +139,21 @@ void FVariableBindingPropertyCustomization::CustomizeChildren(TSharedRef<IProper
 
 void FVariableBindingPropertyCustomization::NotifyPostChange(const FPropertyChangedEvent& PropertyChangedEvent, FEditPropertyChain* PropertyThatChanged)
 {
+	OnBindingChanged();
+}
+
+void FVariableBindingPropertyCustomization::OnBindingChanged()
+{
+	TArray<UObject*> OuterObjects;
+	PropertyHandle->GetOuterObjects(OuterObjects);
+	for(UObject* Object : OuterObjects)
+	{
+		if(UAnimNextVariableEntry* VariableEntry = Cast<UAnimNextVariableEntry>(Object))
+		{
+			VariableEntry->BroadcastModified(EAnimNextEditorDataNotifType::VariableBindingChanged);
+		}
+	}
+	
 	RequestRefresh();
 }
 
@@ -294,6 +316,15 @@ TSharedRef<SWidget> FVariableBindingPropertyCustomization::CreateBindingWidget()
 	return SNullWidget::NullWidget;
 }
 
+void FVariableBindingPropertyCustomization::OnEditorDataModified(ERigVMGraphNotifType NotifType, URigVMGraph* Graph, UObject* Subject, TWeakObjectPtr<UAnimNextVariableEntry> WeakEntry)
+{
+	UAnimNextVariableEntry* Entry = WeakEntry.Get();
+	if (Entry && Subject == Entry)
+	{
+		RequestRefresh();
+	}
+}
+
 }
 
 #undef LOCTEXT_NAMESPACE
\ No newline at end of file
diff --git a/Engine/Plugins/Experimental/AnimNext/Source/AnimNextEditor/Private/Variables/VariableBindingPropertyCustomization.h b/Engine/Plugins/Experimental/AnimNext/Source/AnimNextEditor/Private/Variables/VariableBindingPropertyCustomization.h
--- a/Engine/Plugins/Experimental/AnimNext/Source/AnimNextEditor/Private/Variables/VariableBindingPropertyCustomization.h
+++ b/Engine/Plugins/Experimental/AnimNext/Source/AnimNextEditor/Private/Variables/VariableBindingPropertyCustomization.h
@@ -8,7 +8,10 @@
 #include "Misc/NotifyHook.h"
 #include "Param/ParamType.h"
 
+class URigVMGraph;
+enum class ERigVMGraphNotifType : uint8;
 struct FAnimNextVariableBindingData;
+class UAnimNextVariableEntry;
 
 namespace UE::AnimNext::Editor
 {
@@ -23,7 +26,9 @@ private:
 	// FNotifyHook interface
 	virtual void NotifyPostChange(const FPropertyChangedEvent& PropertyChangedEvent, FEditPropertyChain* PropertyThatChanged) override;
 
-	void Refresh();
+	void OnBindingChanged();
+	void Refresh(); 
+	void OnEditorDataModified(ERigVMGraphNotifType NotifType, URigVMGraph* Graph, UObject* Subject, TWeakObjectPtr<UAnimNextVariableEntry> WeakEntry);
 
 	void RequestRefresh();
 
