AnimNext variables UI/UX changes:

diff --git a/Engine/Plugins/Experimental/AnimNext/Source/AnimNextEditor/Private/Variables/VariableCustomization.cpp b/Engine/Plugins/Experimental/AnimNext/Source/AnimNextEditor/Private/Variables/VariableCustomization.cpp
--- a/Engine/Plugins/Experimental/AnimNext/Source/AnimNextEditor/Private/Variables/VariableCustomization.cpp
+++ b/Engine/Plugins/Experimental/AnimNext/Source/AnimNextEditor/Private/Variables/VariableCustomization.cpp
@@ -8,9 +8,12 @@
 #include "InstancedPropertyBagStructureDataProvider.h"
 #include "UncookedOnlyUtils.h"
 #include "PropertyHandle.h"
-#include "Module/AnimNextModule.h" 
+#include "ScopedTransaction.h"
 #include "Entries/AnimNextVariableEntry.h"
 #include "Module/AnimNextModule_EditorData.h"
+#include "UObject/GarbageCollectionSchema.h"
+#include "Widgets/Input/SComboButton.h"
+#include "Widgets/Input/SEditableTextBox.h"
 
 #define LOCTEXT_NAMESPACE "VariableCustomization"
 
@@ -18,7 +21,7 @@ namespace UE::AnimNext::Editor
 {
 
 void FVariableCustomization::CustomizeDetails(IDetailLayoutBuilder& DetailBuilder)
-{
+{	
 	TArray<TWeakObjectPtr<UObject>> Objects;
 	DetailBuilder.GetObjectsBeingCustomized(Objects);
 
@@ -26,7 +29,9 @@ void FVariableCustomization::CustomizeDetails(IDetailLayoutBuilder& DetailBuilde
 	{
 		return;
 	}
-
+		
+	bool bVariablesArePartOfSameAsset = true;
+	UAnimNextRigVMAsset* MainAsset = nullptr;
 	
 	for(TWeakObjectPtr<UObject> WeakObject : Objects)
 	{
@@ -36,16 +41,84 @@ void FVariableCustomization::CustomizeDetails(IDetailLayoutBuilder& DetailBuilde
 			continue;
 		}
 
+		Variables.Add(Variable);
 
 		UAnimNextRigVMAsset* Asset = Variable->GetTypedOuter<UAnimNextRigVMAsset>();
 
+		if (Asset && MainAsset == nullptr)
+		{
+			MainAsset = Asset;
+		}
+		else if (Asset == nullptr || Asset != MainAsset)
+		{
+			bVariablesArePartOfSameAsset = false;
+		}
+
 		// Disable access specifier switching specifically for shared variables
-		if(UAnimNextSharedVariables* DataInterface = ExactCast<UAnimNextSharedVariables>(Asset))
+		if(ExactCast<UAnimNextSharedVariables>(Asset))
 		{
 			TSharedPtr<IPropertyHandle> AccessProperty = DetailBuilder.GetProperty(GET_MEMBER_NAME_CHECKED(UAnimNextVariableEntry, Access));
 			AccessProperty->MarkHiddenByCustomization();
 		}
 	}
+	
+	TSharedPtr<IPropertyHandle> CategoryProperty = DetailBuilder.GetProperty(GET_MEMBER_NAME_CHECKED(UAnimNextVariableEntry, Category));
+	CategoryProperty->MarkHiddenByCustomization();
+
+	// Only allow changing categories if the variables are part of the same asset
+	if (bVariablesArePartOfSameAsset)
+	{
+		OuterRigVMAsset = MainAsset;
+		PopulateCategories();
+
+		IDetailPropertyRow& CategoryRow = DetailBuilder.AddPropertyToCategory(CategoryProperty);	
+		CategoryRow.CustomWidget()
+		.NameContent()
+		[
+			CategoryProperty->CreatePropertyNameWidget()
+		]
+		.ValueContent()
+		.MinDesiredWidth(250.0f)
+		[
+			SAssignNew(CategoryComboButton, SComboButton)
+			.ContentPadding(FMargin(0.0f, 0.0f, 5.0f, 0.0f))
+			.ButtonContent()
+			[
+				SNew(SBorder)
+				.BorderImage(FAppStyle::Get().GetBrush("NoBorder") )
+				.Padding(FMargin(0, 0, 5, 0))
+				[
+					SNew(SEditableTextBox)
+						.Text(this, &FVariableCustomization::GetCategory)
+						.OnTextCommitted(this, &FVariableCustomization::SetCategory)
+						.OnVerifyTextChanged(this, &FVariableCustomization::OnVerifyCategoryName)
+						//.ToolTip(CategoryProperty->GetToolTipText()) TODO 
+						.SelectAllTextWhenFocused(true)
+						.RevertTextOnEscape(true)
+						.Font( DetailBuilder.GetDetailFont() )
+				]
+			]
+			.MenuContent()
+			[
+				SNew(SVerticalBox)
+				+SVerticalBox::Slot()
+				.AutoHeight()
+				.MaxHeight(400.0f)
+				[
+					SAssignNew(CategoryListView, SListView<TSharedPtr<FText>>)
+						.ListItemsSource(&Categories)
+						.OnGenerateRow_Lambda([](const TSharedPtr<FText>& Item, const TSharedRef<STableViewBase>& OwnerTable)
+						{
+							return SNew(STableRow<TSharedPtr<FString>>, OwnerTable)
+							[
+								SNew(STextBlock).Text(*Item.Get())
+							];
+						})
+						.OnSelectionChanged(this, &FVariableCustomization::OnCategorySelectionChanged)
+				]
+			]
+		];
+	}
 
 	// Don't customize default value if we have multi-selection
 	if(Objects.Num() > 1)
@@ -55,8 +128,6 @@ void FVariableCustomization::CustomizeDetails(IDetailLayoutBuilder& DetailBuilde
 
 	if (UAnimNextVariableEntry* Variable = Cast<UAnimNextVariableEntry>(Objects[0].Get()))
 	{
-		IDetailCategoryBuilder& VariablesCategory = DetailBuilder.EditCategory(TEXT("Variables"), FText::GetEmpty(), ECategoryPriority::Default);
-		
 		IDetailCategoryBuilder& DefaultValueCategory = DetailBuilder.EditCategory(TEXT("DefaultValue"), FText::GetEmpty(), ECategoryPriority::Default);
 
 		TSharedRef< SWidget > ColumnWidget = SNullWidget::NullWidget;
@@ -69,6 +140,9 @@ void FVariableCustomization::CustomizeDetails(IDetailLayoutBuilder& DetailBuilde
 		if (PropertyDesc != nullptr)
 		{
 			IDetailPropertyRow* DetailPropertyRow = DefaultValueCategory.AddExternalStructureProperty(MakeShared<FInstancePropertyBagStructureDataProvider>(PropertyBag), IAnimNextRigVMVariableInterface::ValueName, EPropertyLocation::Default, AddPropertyParams);
+			
+			DetailPropertyRow->ShouldAutoExpand(true);
+			
 			if (TSharedPtr<IPropertyHandle> Handle = DetailPropertyRow->GetPropertyHandle(); Handle.IsValid())
 			{
 				Handle->SetPropertyDisplayName(FText::FromName(Variable->GetEntryName()));
@@ -108,7 +182,99 @@ bool FVariableCustomization::OnVerifyNameChanged(const FText& InText, FText& Out
 	return true;
 }
 
+FText FVariableCustomization::GetCategory() const
+{
+	FString CommonCategory;
+
+	for(int32 Index = 0; Index < Variables.Num(); Index++)
+	{
+		TWeakObjectPtr<UAnimNextVariableEntry> WeakVariable = Variables[Index];
+		UAnimNextVariableEntry* PinnedVariable = WeakVariable.Get();
+		check(PinnedVariable);
+
+		if (Index == 0)
+		{
+			CommonCategory = PinnedVariable->GetVariableCategory(); 	
+		}
+		else
+		{
+			if (CommonCategory != PinnedVariable->GetVariableCategory())
+			{
+				return LOCTEXT("MultipleValuesLabel", "Multiple Values");
+			}
+		}
+	}
+
+	return FText::FromString(CommonCategory);
+}
+
+void FVariableCustomization::SetCategory(const FText& InNewCategory, ETextCommit::Type InCommitType)
+{
+	if(const UAnimNextRigVMAsset* RigVMAsset = OuterRigVMAsset.Get())
+	{
+		UAnimNextRigVMAssetEditorData* EditorData = UncookedOnly::FUtils::GetEditorData<UAnimNextRigVMAssetEditorData>(RigVMAsset);
+		const FString CategoryString = InNewCategory.ToString();
 
+		const FScopedTransaction Transaction(LOCTEXT("ChangeVariableCategory", "Change Variable Category"));
+
+		if (!EditorData->VariableAndFunctionCategories.Contains(CategoryString))
+		{
+			EditorData->AddCategory(CategoryString);
+		}
+		
+		SetCategory(CategoryString);
+	}
+}
+
+bool FVariableCustomization::OnVerifyCategoryName(const FText& InText, FText& OutErrorMessage)
+{
+	if (InText.IsEmpty())
+	{
+		OutErrorMessage = LOCTEXT("CategoryEmpty", "Cannot add a category with an empty string.");
+		return false;
+	}
+	return true;
+}
+
+void FVariableCustomization::OnCategorySelectionChanged(TSharedPtr<FText> ProposedSelection, ESelectInfo::Type)
+{
+	if (ProposedSelection.IsValid())
+	{
+		const FScopedTransaction Transaction(LOCTEXT("ChangeVariableCategory", "Change Variable Category"));
+		
+		const FString CategoryString = ProposedSelection->ToString();
+		SetCategory(CategoryString);
+
+		CategoryListView->ClearSelection();
+		CategoryComboButton->SetIsOpen(false);
+	}
+}
+
+void FVariableCustomization::PopulateCategories()
+{
+	Categories.Empty();
+	if(const UAnimNextRigVMAsset* RigVMAsset = OuterRigVMAsset.Get())
+	{
+		UAnimNextRigVMAssetEditorData* EditorData = UncookedOnly::FUtils::GetEditorData<UAnimNextRigVMAssetEditorData>(RigVMAsset);
+		Algo::Transform(EditorData->VariableAndFunctionCategories, Categories, [](const FString& InCategory)
+		{
+			return MakeShared<FText>(FText::FromString(InCategory));
+		});
+	}
+}
+
+void FVariableCustomization::SetCategory(const FString& CategoryName)
+{
+	for(TWeakObjectPtr<UAnimNextVariableEntry> WeakVariable : Variables)
+	{
+		if(UAnimNextVariableEntry* PinnedVariable = Cast<UAnimNextVariableEntry>(WeakVariable.Get()))
+		{
+			PinnedVariable->SetVariableCategory(CategoryName);
+		}
+	}
+	
+	PopulateCategories();
+}
 }
 
 #undef LOCTEXT_NAMESPACE
diff --git a/Engine/Plugins/Experimental/AnimNext/Source/AnimNextEditor/Private/Variables/VariableCustomization.h b/Engine/Plugins/Experimental/AnimNext/Source/AnimNextEditor/Private/Variables/VariableCustomization.h
--- a/Engine/Plugins/Experimental/AnimNext/Source/AnimNextEditor/Private/Variables/VariableCustomization.h
+++ b/Engine/Plugins/Experimental/AnimNext/Source/AnimNextEditor/Private/Variables/VariableCustomization.h
@@ -4,7 +4,11 @@
 
 #include "IDetailCustomization.h"
 #include "IPropertyTypeCustomization.h"
-#include "PropertyEditorModule.h"
+#include "Widgets/Input/SComboButton.h"
+#include "Widgets/Views/SListView.h"
+
+class UAnimNextRigVMAsset;
+class UAnimNextVariableEntry;
 
 namespace UE::AnimNext::Editor
 {
@@ -12,7 +16,6 @@ namespace UE::AnimNext::Editor
 class FVariableCustomization : public IDetailCustomization
 {
 private:
-
 	/** Called when details should be customized */
 	virtual void CustomizeDetails(IDetailLayoutBuilder& DetailBuilder) override;
 
@@ -25,6 +28,20 @@ private:
 	FText GetName() const;
 	void SetName(const FText& InNewText, ETextCommit::Type InCommitType);
 	bool OnVerifyNameChanged(const FText& InText, FText& OutErrorMessage);
+
+	FText GetCategory() const;
+	void SetCategory(const FText& InNewCategory, ETextCommit::Type InCommitType);	
+	bool OnVerifyCategoryName(const FText& InText, FText& OutErrorMessage);
+	void OnCategorySelectionChanged(TSharedPtr<FText> ProposedSelection, ESelectInfo::Type);
+	void PopulateCategories();
+	void SetCategory(const FString& CategoryName);
+
+	TSharedPtr<SListView<TSharedPtr<FText>>> CategoryListView;
+	TSharedPtr<SComboButton> CategoryComboButton;
+	TArray<TSharedPtr<FText>> Categories;
+
+	TArray<TWeakObjectPtr<UAnimNextVariableEntry>> Variables;
+	TWeakObjectPtr<UAnimNextRigVMAsset> OuterRigVMAsset;
 };
 
 }
diff --git a/Engine/Plugins/Experimental/AnimNext/Source/AnimNextEditor/Private/Workspace/AnimNextWorkspaceEditorMode.cpp b/Engine/Plugins/Experimental/AnimNext/Source/AnimNextEditor/Private/Workspace/AnimNextWorkspaceEditorMode.cpp
--- a/Engine/Plugins/Experimental/AnimNext/Source/AnimNextEditor/Private/Workspace/AnimNextWorkspaceEditorMode.cpp
+++ b/Engine/Plugins/Experimental/AnimNext/Source/AnimNextEditor/Private/Workspace/AnimNextWorkspaceEditorMode.cpp
@@ -7,11 +7,11 @@
 #include "AnimNextRigVMAsset.h"
 #include "AnimNextRigVMAssetEditorData.h"
 #include "AnimNextScopedCompileJob.h"
-#include "AnimNextWorkspaceSchema.h"
 #include "ContextObjectStore.h"
 #include "IAssetCompilationHandler.h"
 #include "InteractiveToolManager.h"
 #include "IWorkspaceEditor.h"
+#include "IWorkspaceEditorModule.h"
 #include "RigVMCommands.h"
 #include "UncookedOnlyUtils.h"
 #include "Logging/MessageLog.h"
@@ -40,7 +40,7 @@ void UAnimNextWorkspaceEditorMode::Enter()
 		return;
 	}
 	
-	WorkspaceEditor->OnFocussedDocumentChanged().AddUObject(this, &UAnimNextWorkspaceEditorMode::HandleFocussedDocumentChanged);
+	WorkspaceEditor->OnFocusedDocumentChanged().AddUObject(this, &UAnimNextWorkspaceEditorMode::HandleFocusedDocumentChanged);
 
 	TArray<UObject*> Assets;
 	WorkspaceEditor->GetOpenedAssets<UAnimNextRigVMAsset>(Assets);
@@ -147,10 +147,10 @@ void UAnimNextWorkspaceEditorMode::HandleCompile()
 			JobName = FText::FromName(Asset->GetFName());
 		}
 	}
-	else if(UObject* FocussedDocument = WorkspaceEditor->GetFocussedDocument())
+	else if(UObject* FocusedDocument = WorkspaceEditor->GetFocusedWorkspaceDocument().GetObject())
 	{
 		// Find the outer asset for this document
-		Asset = FocussedDocument;
+		Asset = FocusedDocument;
 		
 		while(Asset && (!Asset->IsAsset() || Asset->IsPackageExternal()))
 		{
@@ -250,8 +250,10 @@ bool UAnimNextWorkspaceEditorMode::IsCompileDirtyFilesChecked() const
 	return State.bCompileDirtyFiles;
 }
 
-void UAnimNextWorkspaceEditorMode::HandleFocussedDocumentChanged(TObjectPtr<UObject> InObject)
+void UAnimNextWorkspaceEditorMode::HandleFocusedDocumentChanged(const UE::Workspace::FWorkspaceDocument& InDocument)
 {
+	UObject* InObject = InDocument.GetObject();
+	
 	if(InObject == nullptr)
 	{
 		return;
@@ -357,7 +359,7 @@ void UAnimNextWorkspaceEditorMode::HandleRigVMModifiedEvent(ERigVMGraphNotifType
 	}
 }
 
-TSharedPtr<UE::Workspace::IWorkspaceEditor> UAnimNextWorkspaceEditorMode::GetWorkspaceEditor()
+TSharedPtr<UE::Workspace::IWorkspaceEditor> UAnimNextWorkspaceEditorMode::GetWorkspaceEditor() const
 {
 	if (const UContextObjectStore* ContextObjectStore = GetToolManager()->GetContextObjectStore())
 	{
diff --git a/Engine/Plugins/Experimental/AnimNext/Source/AnimNextEditor/Private/Workspace/AnimNextWorkspaceEditorMode.h b/Engine/Plugins/Experimental/AnimNext/Source/AnimNextEditor/Private/Workspace/AnimNextWorkspaceEditorMode.h
--- a/Engine/Plugins/Experimental/AnimNext/Source/AnimNextEditor/Private/Workspace/AnimNextWorkspaceEditorMode.h
+++ b/Engine/Plugins/Experimental/AnimNext/Source/AnimNextEditor/Private/Workspace/AnimNextWorkspaceEditorMode.h
@@ -15,7 +15,8 @@ enum class ERigVMGraphNotifType : uint8;
 
 namespace UE::Workspace
 {
-	class IWorkspaceEditor;
+struct FWorkspaceDocument;
+class IWorkspaceEditor;
 }
 
 UCLASS(Transient)
@@ -32,7 +33,7 @@ public:
 	UE::AnimNext::Editor::ECompileStatus GetLatestCompileStatus() const { return CompileStatus; }
 
 	// Get the currently stored workspace state
-	const FAnimNextWorkspaceState& GetState() { return State; } 
+	const FAnimNextWorkspaceState& GetState() const { return State; } 
 
 private:
 	// UEdMode interface
@@ -60,12 +61,12 @@ private:
 	void HandleRigVMModifiedEvent(ERigVMGraphNotifType InType, URigVMGraph* InGraph, UObject* InSubject);
 
 	// Handle the document focus changing
-	void HandleFocussedDocumentChanged(TObjectPtr<UObject> InAsset); 
+	void HandleFocusedDocumentChanged(const UE::Workspace::FWorkspaceDocument& InDocument); 
 
 	// Forward auto-compilation to all assets in the workspace
 	void PropagateAutoCompile(TSharedRef<UE::Workspace::IWorkspaceEditor> InWorkspaceEditor, bool bInAutoCompile);
 
-	TSharedPtr<UE::Workspace::IWorkspaceEditor> GetWorkspaceEditor();
+	TSharedPtr<UE::Workspace::IWorkspaceEditor> GetWorkspaceEditor() const;
 
 	// Lazily construct an asset compiler for the supplied asset
 	TSharedPtr<UE::AnimNext::Editor::IAssetCompilationHandler> GetAssetCompiler(UObject* InAsset);
diff --git a/Engine/Plugins/Experimental/AnimNext/Source/AnimNextEditor/Public/IAnimNextEditorModule.h b/Engine/Plugins/Experimental/AnimNext/Source/AnimNextEditor/Public/IAnimNextEditorModule.h
--- a/Engine/Plugins/Experimental/AnimNext/Source/AnimNextEditor/Public/IAnimNextEditorModule.h
+++ b/Engine/Plugins/Experimental/AnimNext/Source/AnimNextEditor/Public/IAnimNextEditorModule.h
@@ -4,7 +4,6 @@
 
 #include "CoreMinimal.h"
 #include "Modules/ModuleManager.h"
-#include "UObject/TopLevelAssetPath.h"
 
 class SWidget;
 class IMessageLogListing;
@@ -26,6 +25,7 @@ namespace UE::AnimNext::Editor
 const FLazyName CompilerResultsTabName("CompilerResultsTab");
 const FLazyName LogListingName("AnimNextCompilerResults");
 const FLazyName FindTabName("FindTab");
+const FLazyName FindAndReplaceTabName("FindAndReplaceTab");
 
 struct FVariablePickerArgs;
 struct FActionMenuContextData;
@@ -37,7 +37,7 @@ using FAssetCompilationHandlerFactoryDelegate = TDelegate<TSharedRef<IAssetCompi
 class IAnimNextEditorModule : public IModuleInterface
 {
 public:
-	DECLARE_MULTICAST_DELEGATE_ThreeParams(FCollectGraphMenuActionsMulticast, const TWeakPtr<UE::Workspace::IWorkspaceEditor>& /*IWorkspaceEditor*/, FGraphContextMenuBuilder& /*InContexMenuBuilder*/, const FActionMenuContextData& /*ActionMenuContextData*/);
+	DECLARE_MULTICAST_DELEGATE_ThreeParams(FCollectGraphMenuActionsMulticast, const TWeakPtr<UE::Workspace::IWorkspaceEditor>& /*IWorkspaceEditor*/, FGraphContextMenuBuilder& /*InContextMenuBuilder*/, const FActionMenuContextData& /*ActionMenuContextData*/);
 	typedef FCollectGraphMenuActionsMulticast::FDelegate FOnCollectGraphMenuActionsDelegate;
 
 	DECLARE_MULTICAST_DELEGATE_TwoParams(FNodeDblClickNotificationMulticast, const UE::Workspace::FWorkspaceEditorContext& /*InWorkspaceEditorContext*/, const UEdGraphNode* /*InNode*/);
diff --git a/Engine/Plugins/Experimental/AnimNext/Source/AnimNextUncookedOnly/Internal/AnimNextAssetWorkspaceAssetUserData.h b/Engine/Plugins/Experimental/AnimNext/Source/AnimNextUncookedOnly/Internal/AnimNextAssetWorkspaceAssetUserData.h
--- a/Engine/Plugins/Experimental/AnimNext/Source/AnimNextUncookedOnly/Internal/AnimNextAssetWorkspaceAssetUserData.h
+++ b/Engine/Plugins/Experimental/AnimNext/Source/AnimNextUncookedOnly/Internal/AnimNextAssetWorkspaceAssetUserData.h
@@ -46,13 +46,13 @@ struct FAnimNextModuleOutlinerData : public FAnimNextRigVMAssetOutlinerData
 };
 
 USTRUCT()
-struct FAnimNextDataInterfaceOutlinerData : public FAnimNextRigVMAssetOutlinerData
+struct FAnimNextSharedVariablesOutlinerData : public FAnimNextRigVMAssetOutlinerData
 {
 	GENERATED_BODY()
 
-	FAnimNextDataInterfaceOutlinerData() = default;
+	FAnimNextSharedVariablesOutlinerData() = default;
 
-	UAnimNextSharedVariables* GetDataInterface() const
+	UAnimNextSharedVariables* GetSharedVariables() const
 	{
 		return Cast<UAnimNextSharedVariables>(GetAsset());
 	}
diff --git a/Engine/Plugins/Experimental/AnimNext/Source/AnimNextUncookedOnly/Internal/AnimNextControllerBase.h b/Engine/Plugins/Experimental/AnimNext/Source/AnimNextUncookedOnly/Internal/AnimNextControllerBase.h
--- a/Engine/Plugins/Experimental/AnimNext/Source/AnimNextUncookedOnly/Internal/AnimNextControllerBase.h
+++ b/Engine/Plugins/Experimental/AnimNext/Source/AnimNextUncookedOnly/Internal/AnimNextControllerBase.h
@@ -1,4 +1,4 @@
-﻿// Copyright Epic Games, Inc. All Rights Reserved.
+// Copyright Epic Games, Inc. All Rights Reserved.
 
 #pragma once
 
@@ -27,12 +27,14 @@ public:
 	UFUNCTION(BlueprintCallable, Category = RigVMController)
 	ANIMNEXTUNCOOKEDONLY_API UAnimNextSharedVariableNode* AddStructSharedVariableNode(const UScriptStruct* InStruct, const FName& InVariableName, const FString& InCPPType, UObject* InCPPTypeObject, bool bIsGetter, const FString& InDefaultValue, const FVector2D& InPosition = FVector2D::ZeroVector, const FString& InNodeName = TEXT(""), bool bSetupUndoRedo = true, bool bPrintPythonCommand = false);
 
-	// Add a shared variable from the specified source object
-	ANIMNEXTUNCOOKEDONLY_API UAnimNextSharedVariableNode* AddSharedVariableNode(const FString& InSourceObjectPath, const FName& InVariableName, const FString& InCPPType, UObject* InCPPTypeObject, bool bIsGetter, const FString& InDefaultValue, const FVector2D& InPosition, const FString& InNodeName, bool bSetupUndoRedo, bool bPrintPythonCommand);
-
-private:
-	friend UAnimNextRigVMAssetEditorData;
+	// Refreshes a UAnimNextSharedVariableNode instance with provided data (similar to URigVMController::RefreshVariableNode)
+	UFUNCTION(BlueprintCallable, Category = RigVMController)
+	ANIMNEXTUNCOOKEDONLY_API void RefreshSharedVariableNode(const FName& InNodeName, const FString& InSourceObjectPath, const FName& InVariableName, const FString& InCPPType, UObject* InCPPTypeObject, bool bSetupUndoRedo, bool bPrintPythonCommand);
 
-	// Helper function to upgrade older assets
-	void ReplaceVariableNodeWithSharedVariableNode(URigVMVariableNode* InVariableNode, FName InNewVariableName, const UObject* InAssetOrStruct);
+	// Add a shared variable from the specified source object
+	UFUNCTION(BlueprintCallable, Category = RigVMController)
+	ANIMNEXTUNCOOKEDONLY_API UAnimNextSharedVariableNode* AddSharedVariableNode(const FString& InSourceObjectPath, const FName& InVariableName, const FString& InCPPType, const FString& InCPPTypeObjectPath, bool bIsGetter, const FString& InDefaultValue, const FVector2D& InPosition, const FString& InNodeName, bool bSetupUndoRedo, bool bPrintPythonCommand);
+	
+	UAnimNextSharedVariableNode* ReplaceVariableNodeWithSharedVariableNode(URigVMVariableNode* InVariableNode, FName InNewVariableName, const UObject* InAssetOrStruct, bool bSetupUndoRedo, bool bPrintPythonCommand);
+	URigVMVariableNode* ReplaceSharedVariableNodeWithVariableNode(UAnimNextSharedVariableNode* InVariableNode, FName InNewVariableName, bool bSetupUndoRedo, bool bPrintPythonCommand);
 };
diff --git a/Engine/Plugins/Experimental/AnimNext/Source/AnimNextUncookedOnly/Internal/AnimNextRigVMAssetEditorData.h b/Engine/Plugins/Experimental/AnimNext/Source/AnimNextUncookedOnly/Internal/AnimNextRigVMAssetEditorData.h
--- a/Engine/Plugins/Experimental/AnimNext/Source/AnimNextUncookedOnly/Internal/AnimNextRigVMAssetEditorData.h
+++ b/Engine/Plugins/Experimental/AnimNext/Source/AnimNextUncookedOnly/Internal/AnimNextRigVMAssetEditorData.h
@@ -60,6 +60,8 @@ namespace UE::AnimNext::Editor
 	class FAnimNextAnimGraphEditorModule;
 	class FCallFunctionSharedDataDetails;
 	class FAssetCompilationHandler;
+	class FVariableBindingPropertyCustomization;
+	struct FVariablesOutlinerCategoryItem;
 }
 
 namespace UE::AnimNext::Tests
@@ -84,6 +86,10 @@ enum class EAnimNextEditorDataNotifType : uint8
 	UndoRedo,		// Transaction was performed (Subject == UObject)
 	VariableDefaultValueChanged,	// A variable entry default value changed (Subject == UAnimNextVariableEntry)
 	VariableBindingChanged,	// A variable entry binding changed (Subject == UAnimNextVariableEntry)
+	VariableCategoryChanged, // A variable entry its category changed (Subject == UAnimNextVariableEntry)
+	CategoryAdded, // A category has been added (Subject == UAnimNextRigVMAssetEditorData)
+	CategoryChanged, // A category has been changed (Subject == UAnimNextRigVMAssetEditorData)
+	VariablesReordered,  // Asset entries have been reordered (Subject == UAnimNextRigVMAssetEditorData)
 };
 
 namespace UE::AnimNext::UncookedOnly
@@ -123,7 +129,7 @@ public:
 	static ANIMNEXTUNCOOKEDONLY_API UAnimNextVariableEntry* AddVariable(UAnimNextRigVMAsset* InAsset, FName InName, EPropertyBagPropertyType InValueType, EPropertyBagContainerType InContainerType = EPropertyBagContainerType::None, const UObject* InValueTypeObject = nullptr, const FString& InDefaultValue = TEXT(""), bool bSetupUndoRedo = true, bool bPrintPythonCommand = true);
 
 	/** Adds an event graph to an AnimNext asset */
-	UFUNCTION(BlueprintCallable, Category = "AnimNext|Entries", meta=(ScriptMethod))
+	UFUNCTION(BlueprintCallable, Category = "UAF|Entries", meta=(ScriptMethod))
 	static ANIMNEXTUNCOOKEDONLY_API UAnimNextEventGraphEntry* AddEventGraph(UAnimNextRigVMAsset* InAsset, FName InName, UScriptStruct* InEventStruct, bool bSetupUndoRedo = true, bool bPrintPythonCommand = true);
 
 	/** Adds shared variables to an AnimNext asset */
@@ -131,12 +137,20 @@ public:
 	static ANIMNEXTUNCOOKEDONLY_API UAnimNextSharedVariablesEntry* AddSharedVariables(UAnimNextRigVMAsset* InAsset, UAnimNextSharedVariables* InSharedVariables, bool bSetupUndoRedo = true, bool bPrintPythonCommand = true);
 
 	/** Adds shared variables struct to an AnimNext asset */
-	UFUNCTION(BlueprintCallable, Category = "AnimNext|Entries", meta=(ScriptMethod))
+	UFUNCTION(BlueprintCallable, Category = "UAF|Entries", meta=(ScriptMethod))
 	static ANIMNEXTUNCOOKEDONLY_API UAnimNextSharedVariablesEntry* AddSharedVariablesStruct(UAnimNextRigVMAsset* InAsset, UScriptStruct* InStruct, bool bSetupUndoRedo = true, bool bPrintPythonCommand = true);
 	
 	/** Adds a function to an AnimNext asset */
 	UFUNCTION(BlueprintCallable, Category = "UAF|Entries", meta=(ScriptMethod))
 	static ANIMNEXTUNCOOKEDONLY_API URigVMLibraryNode* AddFunction(UAnimNextRigVMAsset* InAsset, FName InFunctionName, bool bInMutable, bool bSetupUndoRedo = true, bool bPrintPythonCommand = true);
+
+	/** Adds a category to an AnimNext asset */
+	UFUNCTION(BlueprintCallable, Category = "UAF|Entries", meta=(ScriptMethod))
+	static ANIMNEXTUNCOOKEDONLY_API bool AddCategory(UAnimNextRigVMAsset* InAsset, const FString& CategoryName, bool bSetupUndoRedo = true, bool bPrintPythonCommand = true);
+
+	/** Renames a category within an AnimNext asset */
+	UFUNCTION(BlueprintCallable, Category = "UAF|Entries", meta=(ScriptMethod))
+	static ANIMNEXTUNCOOKEDONLY_API bool RenameCategory(UAnimNextRigVMAsset* InAsset, const FString& CategoryName, const FString& NewCategoryName, bool bSetupUndoRedo = true, bool bPrintPythonCommand = true);
 };
 
 /* Base class for all AnimNext editor data objects that use RigVM */
@@ -164,6 +178,18 @@ public:
 	// Find an entry by name
 	UE_API UAnimNextRigVMAssetEntry* FindEntry(FName InName) const;
 
+	/** Adds a category for variables/functions to this asset */
+	UE_API bool AddCategory(const FString& CategoryName, bool bInSetupUndoRedo = true, bool bPrintPythonCommand = true);
+
+	/** Renames a category using provided new name */
+	UE_API bool RenameCategory(const FString& OldName, const FString& NewName, bool bInSetupUndoRedo = true, bool bPrintPythonCommand = true);
+
+	/** Reorders category entry to appear before provided BeforeCategoryName */
+	UE_API void ReorderCategory(const FString& CategoryName, const FString& BeforeCategoryName, bool bInSetupUndoRedo = true, bool bPrintPythonCommand = true);
+
+	/** Reorders variable entry to appear before provided BeforeVariableEntry */
+	UE_API void ReorderVariable(UAnimNextVariableEntry* VariableEntry, const UAnimNextVariableEntry* BeforeVariableEntry, bool bSetupUndoRedo = true, bool bPrintPythonCommand = true);
+
 	// Get the external packaging status of this asset
 	bool IsUsingExternalPackages() const { return bUsesExternalPackages; }
 
@@ -196,12 +222,12 @@ protected:
 	friend class UE::AnimNext::Editor::SFindInAnimNextRigVMAsset;
 	friend class UE::AnimNext::Tests::FVariables;
 	friend class UE::AnimNext::Tests::FVariables_UOLBindings;
-	friend class UE::AnimNext::Tests::FDataInterfaceCompile;
 	friend class UE::AnimNext::Editor::FVariableCustomization;
 	friend class UE::AnimNext::Editor::FAnimNextAssetItemDetails;
 	friend class UE::AnimNext::Editor::FAnimNextGraphItemDetails;
 	friend class UE::AnimNext::Editor::FAnimNextFunctionItemDetails;
 	friend struct UE::AnimNext::Editor::FVariablesOutlinerEntryItem;
+	friend struct UE::AnimNext::Editor::FVariablesOutlinerCategoryItem;
 	friend class UE::AnimNext::Editor::FVariablesOutlinerMode;
 	friend class UE::AnimNext::Editor::FVariablesOutlinerHierarchy;
 	friend class UE::AnimNext::Editor::SVariablesOutlinerValue;
@@ -217,6 +243,9 @@ protected:
 	friend UE::AnimNext::UncookedOnly::FScopedCompileJob;
 	friend UE::AnimNext::Editor::FAssetCompilationHandler;
 	friend class UAnimNextStateTreeTreeEditorData;
+	friend class UE::AnimNext::Editor::FVariableBindingPropertyCustomization;
+	friend class UAnimNextAssetFindReplaceVariables;
+	friend struct FAnimNextSchemaAction_Variable;
 	friend struct FAnimNextModuleTests;
 	friend struct FAnimNextAnimationGraphTests;
 
@@ -233,7 +262,7 @@ protected:
 	UE_API void HandlePackageDone(const FEndLoadPackageContext& Context);
 	UE_API void HandlePackageDone();
 
-	virtual void GetAnimNextAssetRegistryTags(FAssetRegistryTagsContext& Context, FAnimNextAssetRegistryExports& OutExports) const {}
+	UE_API virtual void GetAnimNextAssetRegistryTags(FAssetRegistryTagsContext& Context, FAnimNextAssetRegistryExports& OutExports) const;
 
 	// IRigVMClientHost interface
 	virtual FString GetAssetName() const override { return GetName(); }
@@ -500,6 +529,19 @@ protected:
 		EAnimNextExportAccessSpecifier Access = EAnimNextExportAccessSpecifier::Private;
 		const FProperty* Property = nullptr;
 		TConstArrayView<uint8> DefaultValue;
+
+		friend uint32 GetTypeHash(const FVariableInfo& Info)
+		{
+			return HashCombine(HashCombine(GetTypeHash(Info.Name), GetTypeHash(Info.Type)), GetTypeHash(Info.SourceObject));
+		}
+
+		bool operator==(const FVariableInfo& Other) const
+		{
+			return Name == Other.Name
+				&& Type == Other.Type
+				&& SourceObject == Other.SourceObject
+				&& Access == Other.Access;
+		}
 	};
 
 	// How to recurse when getting variable info
@@ -536,6 +578,9 @@ protected:
 	// Add a new entry to this asset, taking into account external packaging status
 	UE_API void AddEntryInternal(UAnimNextRigVMAssetEntry* InEntry);
 
+	// Inserts a new entry to this asset, at the provided index, while taking into account external packaging status
+	void InsertEntryInternal(UAnimNextRigVMAssetEntry* InEntry, int32 InsertionIndex);
+
 	// Remove an entry to this asset, taking into account external packaging status
 	UE_API void RemoveEntryInternal(UAnimNextRigVMAssetEntry* InEntry);
 
@@ -624,6 +669,10 @@ protected:
 	UPROPERTY(EditAnywhere, DisplayName = "Default Injection Site", Category=Injection, meta=(AllowedType = FAnimNextAnimGraph))
 	FAnimNextVariableReference DefaultInjectionSiteReference;
 
+	// Set of category names, assignable to variables/functions by user
+	UPROPERTY()
+	TArray<FString> VariableAndFunctionCategories;
+	
 	FOnRigVMCompiledEvent RigVMCompiledEvent;
 
 	FRigVMGraphModifiedEvent RigVMGraphModifiedEvent;
diff --git a/Engine/Plugins/Experimental/AnimNext/Source/AnimNextUncookedOnly/Internal/AnimNextSharedVariableNode.h b/Engine/Plugins/Experimental/AnimNext/Source/AnimNextUncookedOnly/Internal/AnimNextSharedVariableNode.h
--- a/Engine/Plugins/Experimental/AnimNext/Source/AnimNextUncookedOnly/Internal/AnimNextSharedVariableNode.h
+++ b/Engine/Plugins/Experimental/AnimNext/Source/AnimNextUncookedOnly/Internal/AnimNextSharedVariableNode.h
@@ -1,4 +1,4 @@
-﻿// Copyright Epic Games, Inc. All Rights Reserved.
+// Copyright Epic Games, Inc. All Rights Reserved.
 
 #pragma once
 
@@ -9,6 +9,16 @@
 class UAnimNextSharedVariables;
 class UAnimNextControllerBase;
 
+namespace UE::AnimNext::UncookedOnly
+{
+struct FUtils;
+}
+
+namespace UE::AnimNext::Editor
+{
+class FAnimNextEditorModule;
+}
+
 UCLASS()
 class UAnimNextSharedVariableNode : public URigVMVariableNode
 {
@@ -19,6 +29,8 @@ class UAnimNextSharedVariableNode : public URigVMVariableNode
 
 private:
 	friend UAnimNextControllerBase;
+	friend UE::AnimNext::UncookedOnly::FUtils;
+	friend UE::AnimNext::Editor::FAnimNextEditorModule;
 
 	UPROPERTY()
 	TObjectPtr<const UAnimNextSharedVariables> Asset;
diff --git a/Engine/Plugins/Experimental/AnimNext/Source/AnimNextUncookedOnly/Internal/Entries/AnimNextCategoryEntry.h b/Engine/Plugins/Experimental/AnimNext/Source/AnimNextUncookedOnly/Internal/Entries/AnimNextCategoryEntry.h
new file mode 100644
--- /dev/null
+++ b/Engine/Plugins/Experimental/AnimNext/Source/AnimNextUncookedOnly/Internal/Entries/AnimNextCategoryEntry.h
@@ -0,0 +1,33 @@
+// Copyright Epic Games, Inc. All Rights Reserved.
+
+#pragma once
+
+#include "CoreMinimal.h"
+#include "AnimNextRigVMAssetEntry.h"
+
+#include "AnimNextCategoryEntry.generated.h"
+
+namespace UE::AnimNext::Editor
+{
+struct FVariablesOutlinerCategoryItem;
+class FVariablesOutlinerHierarchy;
+}
+
+UCLASS(MinimalAPI, DisplayName = "Category")
+class UAnimNextCategoryEntry : public UAnimNextRigVMAssetEntry
+{
+	GENERATED_BODY()
+
+	friend class UAnimNextRigVMAssetEditorData;
+	friend UE::AnimNext::Editor::FVariablesOutlinerCategoryItem;
+	friend UE::AnimNext::Editor::FVariablesOutlinerHierarchy;
+
+	// UAnimNextRigVMAssetEntry interface
+	virtual void Initialize(UAnimNextRigVMAssetEditorData* InEditorData) override;	
+	virtual FName GetEntryName() const override { return CategoryName; }
+	virtual void SetEntryName(FName InName, bool bSetupUndoRedo = true) override;
+
+protected:
+	UPROPERTY(VisibleAnywhere, Category = "Category")
+	FName CategoryName;
+};
\ No newline at end of file
diff --git a/Engine/Plugins/Experimental/AnimNext/Source/AnimNextUncookedOnly/Internal/Entries/AnimNextSharedVariablesEntry.h b/Engine/Plugins/Experimental/AnimNext/Source/AnimNextUncookedOnly/Internal/Entries/AnimNextSharedVariablesEntry.h
--- a/Engine/Plugins/Experimental/AnimNext/Source/AnimNextUncookedOnly/Internal/Entries/AnimNextSharedVariablesEntry.h
+++ b/Engine/Plugins/Experimental/AnimNext/Source/AnimNextUncookedOnly/Internal/Entries/AnimNextSharedVariablesEntry.h
@@ -1,11 +1,10 @@
-﻿// Copyright Epic Games, Inc. All Rights Reserved.
+// Copyright Epic Games, Inc. All Rights Reserved.
 
 #pragma once
 
 #include "CoreMinimal.h"
 #include "AnimNextRigVMAssetEntry.h"
 #include "Param/ParamType.h"
-#include "StructUtils/PropertyBag.h"
 #include "AnimNextSharedVariablesEntry.generated.h"
 
 class UAnimNextSharedVariables;
@@ -15,13 +14,14 @@ class UAssetDefinition_AnimNextSharedVariablesEntry;
 namespace UE::AnimNext::Editor
 {
 	class FVariablesOutlinerHierarchy;
-	struct FVariablesOutlinerDataInterfaceItem;
-	class SVariablesOutlinerDataInterfaceLabel;
+	struct FVariablesOutlinerStructSharedVariablesItem;
+	class SVariablesOutlinerStructSharedVariablesLabel;
 	class SVariablesOutlinerValue;
 	class FVariableProxyCustomization;
 	class SAddVariablesDialog;
 	class SVariableOverride;
 	class FVariablesOutlinerMode;
+	class SVariablesOutliner;
 }
 
 namespace UE::AnimNext::UncookedOnly
@@ -29,10 +29,6 @@ namespace UE::AnimNext::UncookedOnly
 	struct FUtils;
 }
 
-namespace UE::AnimNext::Tests
-{
-	class FDataInterfaceCompile;
-}
 
 UENUM()
 enum class EAnimNextSharedVariablesType : uint8
@@ -49,15 +45,16 @@ class UAnimNextSharedVariablesEntry : public UAnimNextRigVMAssetEntry
 	friend class UAnimNextRigVMAssetEditorData;
 	friend class UAssetDefinition_AnimNextSharedVariablesEntry;
 	friend class UE::AnimNext::Editor::FVariablesOutlinerHierarchy;
-	friend struct UE::AnimNext::Editor::FVariablesOutlinerDataInterfaceItem;
+	friend struct UE::AnimNext::Editor::FVariablesOutlinerStructSharedVariablesItem;
 	friend struct UE::AnimNext::UncookedOnly::FUtils;
-	friend class UE::AnimNext::Tests::FDataInterfaceCompile;
-	friend class UE::AnimNext::Editor::SVariablesOutlinerDataInterfaceLabel;
+	friend class UE::AnimNext::Editor::SVariablesOutlinerStructSharedVariablesLabel;
 	friend class UE::AnimNext::Editor::SVariablesOutlinerValue;
 	friend class UE::AnimNext::Editor::FVariableProxyCustomization;
 	friend class UE::AnimNext::Editor::SAddVariablesDialog;
 	friend class UE::AnimNext::Editor::SVariableOverride;
+	friend class UE::AnimNext::Editor::SVariablesOutliner;
 	friend UE::AnimNext::Editor::FVariablesOutlinerMode;
+	friend class UAnimNextAssetFindReplaceVariables;
 
 	// UAnimNextRigVMAssetEntry interface
 	virtual void Initialize(UAnimNextRigVMAssetEditorData* InEditorData) override;
diff --git a/Engine/Plugins/Experimental/AnimNext/Source/AnimNextUncookedOnly/Internal/Entries/AnimNextVariableEntry.h b/Engine/Plugins/Experimental/AnimNext/Source/AnimNextUncookedOnly/Internal/Entries/AnimNextVariableEntry.h
--- a/Engine/Plugins/Experimental/AnimNext/Source/AnimNextUncookedOnly/Internal/Entries/AnimNextVariableEntry.h
+++ b/Engine/Plugins/Experimental/AnimNext/Source/AnimNextUncookedOnly/Internal/Entries/AnimNextVariableEntry.h
@@ -30,6 +30,9 @@ namespace UE::AnimNext::Editor
 	class SVariablesOutlinerValue;
 	class FVariableProxyCustomization;
 	class FAnimNextEditorModule;
+	class SVariablesOutlinerEntryLabel;
+	class SVariablesOutlinerAssetLabel;
+	class SVariablesOutlinerCategoryLabel;
 }
 
 namespace UE::AnimNext::Tests
@@ -59,7 +62,6 @@ class UAnimNextVariableEntry : public UAnimNextRigVMAssetEntry, public IAnimNext
 	friend class UE::AnimNext::Tests::FEditor_Variables;
 	friend class UE::AnimNext::Tests::FEditor_AnimGraph_Variables;
 	friend class UE::AnimNext::Tests::FVariables_UOLBindings;
-	friend class UE::AnimNext::Tests::FDataInterfaceCompile;
 	friend class FAnimationAnimNextRuntimeTest_Variables;
 	friend class UE::AnimNext::Editor::FVariableCustomization;
 	friend class UE::AnimNext::Editor::FVariableBindingPropertyCustomization;
@@ -72,6 +74,11 @@ class UAnimNextVariableEntry : public UAnimNextRigVMAssetEntry, public IAnimNext
 	friend class UE::AnimNext::Editor::FVariableProxyCustomization;
 	friend struct UE::AnimNext::UncookedOnly::FPublicVariablesImpl;
 	friend class UE::AnimNext::Editor::FAnimNextEditorModule;
+	friend struct FAnimNextSchemaAction_PromoteToVariable;
+	friend class UE::AnimNext::Editor::SVariablesOutlinerEntryLabel;
+	friend class UE::AnimNext::Editor::SVariablesOutlinerAssetLabel;
+	friend class UE::AnimNext::Editor::SVariablesOutlinerCategoryLabel;
+	friend class UAnimNextAssetFindReplaceVariables;
 
 	UAnimNextVariableEntry();
 
@@ -103,9 +110,12 @@ class UAnimNextVariableEntry : public UAnimNextRigVMAssetEntry, public IAnimNext
 	virtual bool SetDefaultValueFromString(const FString& InDefaultValue, bool bSetupUndoRedo = true) override;
 	virtual const FInstancedPropertyBag& GetPropertyBag() const override;
 	virtual bool GetDefaultValue(const FProperty*& OutProperty, TConstArrayView<uint8>& OutValue) const override;
+	virtual bool GetDefaultValueString(FString& OutValueString) const override;
 	virtual void SetBindingType(UScriptStruct* InBindingTypeStruct, bool bSetupUndoRedo = true) override;
 	virtual void SetBinding(TInstancedStruct<FAnimNextVariableBindingData>&& InBinding, bool bSetupUndoRedo = true) override;
 	virtual TConstStructView<FAnimNextVariableBindingData> GetBinding() const override;
+	virtual FStringView GetVariableCategory() const override;
+	virtual void SetVariableCategory(FStringView InCategoryName, bool bSetupUndoRedo = true) override;
 
 	// Set the default value
 	template<typename ValueType>
@@ -149,6 +159,10 @@ class UAnimNextVariableEntry : public UAnimNextRigVMAssetEntry, public IAnimNext
 	UPROPERTY(EditAnywhere, Category = "Variable", meta=(MultiLine))
 	FString Comment;
 
+	/** Category this variable should be assigned to */
+	UPROPERTY(EditAnywhere, Category = "Variable")
+	FString Category;
+
 	/** Property bag holding the default value of the variable */
 	UPROPERTY()
 	FInstancedPropertyBag DefaultValue;
diff --git a/Engine/Plugins/Experimental/AnimNext/Source/AnimNextUncookedOnly/Internal/UncookedOnlyUtils.h b/Engine/Plugins/Experimental/AnimNext/Source/AnimNextUncookedOnly/Internal/UncookedOnlyUtils.h
--- a/Engine/Plugins/Experimental/AnimNext/Source/AnimNextUncookedOnly/Internal/UncookedOnlyUtils.h
+++ b/Engine/Plugins/Experimental/AnimNext/Source/AnimNextUncookedOnly/Internal/UncookedOnlyUtils.h
@@ -3,17 +3,17 @@
 #pragma once
 
 #include "CoreMinimal.h"
+#include "AnimNextAssetWorkspaceAssetUserData.h"
 #include "IAnimNextRigVMGraphInterface.h"
 #include "AssetRegistry/AssetData.h"
-#include "Misc/Build.h"
 #include "Param/ParamType.h"
 #include "RigVMCore/RigVMTemplate.h"
 #include "StructUtils/InstancedStruct.h"
 #include "RigVMCore/RigVMGraphFunctionDefinition.h"
 #include "UObject/SoftObjectPath.h"
 #include "RigVMModel/RigVMGraph.h"
-// --- ---
-#include "UncookedOnlyUtils.generated.h"
+#include "AnimNextExports.h"
+#include "AssetRegistry/IAssetRegistry.h"
 
 #define UE_API ANIMNEXTUNCOOKEDONLY_API
 
@@ -40,143 +40,17 @@ class UAnimNextRigVMAsset;
 class UAnimNextRigVMAssetEditorData;
 class UAnimNextRigVMAssetEntry;
 class URigVMUnitNode;
-
+class UAnimNextControllerBase;
 
 namespace UE
 {
 	namespace AnimNext
 	{
-		static const FLazyName ExportsAnimNextAssetRegistryTag = TEXT("AnimNextExports");
 		static const FLazyName AnimNextPublicGraphFunctionsExportsRegistryTag = TEXT("AnimNextPublicGraphFunctions");
 		static const FLazyName ControlRigAssetPublicGraphFunctionsExportsRegistryTag = TEXT("PublicGraphFunctions");
 	}
 }
 
-UENUM()
-enum class EAnimNextExportedVariableFlags : uint32
-{
-	NoFlags = 0x0,
-	Public = 0x1,
-	Read = 0x02,
-	Write = 0x04,
-	Declared = 0x08,
-	Max
-};
-
-ENUM_CLASS_FLAGS(EAnimNextExportedVariableFlags)
-
-USTRUCT()
-struct FAnimNextAssetRegistryExportedVariable
-{
-	GENERATED_BODY()
-
-	FAnimNextAssetRegistryExportedVariable() = default;
-
-	FAnimNextAssetRegistryExportedVariable(FName InName, const FAnimNextParamType& InType, EAnimNextExportedVariableFlags InFlags = EAnimNextExportedVariableFlags::NoFlags)
-		: Name(InName)
-		, Type(InType)
-		, Flags((uint32)InFlags) 
-	{}
-
-	EAnimNextExportedVariableFlags GetFlags() const
-	{
-		return (EAnimNextExportedVariableFlags)Flags;
-	}
-	
-	UPROPERTY()
-	FName Name;
-
-	UPROPERTY()
-	FAnimNextParamType Type;
-
-	UPROPERTY()
-	uint32 Flags = (int32)EAnimNextExportedVariableFlags::NoFlags;
-};
-
-USTRUCT()
-struct FAnimNextAssetRegistryExportedSharedVariableAsset
-{
-	GENERATED_BODY()
-
-	FAnimNextAssetRegistryExportedSharedVariableAsset() = default;
-
-	FAnimNextAssetRegistryExportedSharedVariableAsset(const UAnimNextSharedVariables* InSharedVariables)
-		: SharedVariablesAsset(InSharedVariables)
-	{}
-
-	UPROPERTY()
-	TObjectPtr<const UAnimNextSharedVariables> SharedVariablesAsset;
-};
-
-USTRUCT()
-struct FAnimNextAssetRegistryExportedSharedVariableStruct
-{
-	GENERATED_BODY()
-
-	FAnimNextAssetRegistryExportedSharedVariableStruct() = default;
-
-	FAnimNextAssetRegistryExportedSharedVariableStruct(const UScriptStruct* InStruct)
-		: SharedVariablesStruct(InStruct)
-	{}
-
-	UPROPERTY()
-	TObjectPtr<const UScriptStruct> SharedVariablesStruct;
-};
-
-USTRUCT()
-struct FAnimNextAssetRegistryManifestNode
-{
-	GENERATED_BODY()
-
-	FAnimNextAssetRegistryManifestNode() = default;
-
-	FAnimNextAssetRegistryManifestNode(const TObjectPtr<URigVMGraph>& InModelGraph, const FString& InNodeName, const FString& InNodeCategory, const FString& InMenuDesc, const FString& InToolTip)
-		: ModelGraph(FSoftObjectPath::ConstructFromObject(InModelGraph))
-		, NodeName(InNodeName)
-		, NodeCategory(InNodeCategory)
-		, MenuDesc(InMenuDesc)
-		, ToolTip(InToolTip)
-	{}
-
-	bool operator==(const FAnimNextAssetRegistryManifestNode& Other) const
-	{
-		return ModelGraph == Other.ModelGraph && NodeName == Other.NodeName;
-	}
-
-	UPROPERTY()
-	FSoftObjectPath ModelGraph;
-
-	UPROPERTY()
-	FString NodeName;
-
-	UPROPERTY()
-	FString NodeCategory;
-
-	UPROPERTY()
-	FString MenuDesc;
-
-	UPROPERTY()
-	FString ToolTip;
-};
-
-USTRUCT()
-struct FAnimNextAssetRegistryExports
-{
-	GENERATED_BODY()
-
-	UPROPERTY()
-	TArray<FAnimNextAssetRegistryExportedVariable> Variables;
-
-	UPROPERTY()
-	TArray<FAnimNextAssetRegistryExportedSharedVariableAsset> SharedVariableAssets;
-
-	UPROPERTY()
-	TArray<FAnimNextAssetRegistryExportedSharedVariableStruct> SharedVariableStructs;
-
-	UPROPERTY()
-	TArray<FAnimNextAssetRegistryManifestNode> ManifestNodes;
-};
-
 namespace UE::AnimNext::UncookedOnly
 {
 
@@ -239,28 +113,70 @@ struct FUtils
 	// Gets all the variables that are exported to the asset registry
 	static UE_API bool GetExportedVariablesFromAssetRegistry(TMap<FAssetData, FAnimNextAssetRegistryExports>& OutExports);
 
-	// Gets the functions that are exported to the asset registry for an asset
-	static UE_API bool GetExportedFunctionsForAsset(const FAssetData& InAsset, FName Tag, FRigVMGraphFunctionHeaderArray& OutExports);
+	template<typename ExportDataType>
+	static bool GetExportsOfTypeForAsset(const FAssetData& InAsset, FAnimNextAssetRegistryExports& OutExports)
+	{
+		const int32 StartNum = OutExports.Exports.Num();
+		FAnimNextAssetRegistryExports AssetExports;
+		const FString TagValue = InAsset.GetTagValueRef<FString>(UE::AnimNext::ExportsAnimNextAssetRegistryTag);
+		if (FAnimNextAssetRegistryExports::StaticStruct()->ImportText(*TagValue, &AssetExports, nullptr, PPF_None, nullptr, FAnimNextAssetRegistryExports::StaticStruct()->GetName()) != nullptr)
+		{
+			Algo::TransformIf(AssetExports.Exports, OutExports.Exports,
+				[](const FAnimNextExport& Export)
+				{
+					return Export.Data.GetScriptStruct() == ExportDataType::StaticStruct();
+				},
+				[](const FAnimNextExport& Export)
+				{
+					return Export;
+				});
+		}
+
+		return OutExports.Exports.Num() > StartNum;
+	}
+
+	template<typename ExportDataType>
+	static bool GetExportsOfTypeFromAssetRegistry(TMap<FAssetData, FAnimNextAssetRegistryExports>& OutExports)
+	{
+		TArray<FAssetData> AssetData;
+		IAssetRegistry::GetChecked().GetAssetsByTags({UE::AnimNext::ExportsAnimNextAssetRegistryTag}, AssetData);
+
+		for (const FAssetData& Asset : AssetData)
+		{
+			const FString TagValue = Asset.GetTagValueRef<FString>(UE::AnimNext::ExportsAnimNextAssetRegistryTag);
+			FAnimNextAssetRegistryExports AssetExports;
+			if (FAnimNextAssetRegistryExports::StaticStruct()->ImportText(*TagValue, &AssetExports, nullptr, PPF_None, nullptr, FAnimNextAssetRegistryExports::StaticStruct()->GetName()) != nullptr)
+			{
+				FAnimNextAssetRegistryExports& Export = OutExports.Add(Asset);
+				Algo::TransformIf(AssetExports.Exports, Export.Exports,
+					[](const FAnimNextExport& Export)
+					{
+						return Export.Data.GetScriptStruct() == ExportDataType::StaticStruct();
+					},
+					[](const FAnimNextExport& Export)
+					{
+						return Export;
+					});
+			}
+		}
+
+		return OutExports.Num() > 0;
+	}
+
+	static void GetAssetVariableExports(const UAnimNextRigVMAssetEditorData* EditorData, FAnimNextAssetRegistryExports& OutExports, FAssetRegistryTagsContext Context);
 
 	// Gets all the functions that are exported to the asset registry for the specified Tag
 	static UE_API bool GetExportedFunctionsFromAssetRegistry(FName Tag, TMap<FAssetData, FRigVMGraphFunctionHeaderArray>& OutExports);
 
-	// Gets the exported variables that are used by a RigVM asset
-	static UE_API void GetAssetVariables(const UAnimNextRigVMAssetEditorData* EditorData, FAnimNextAssetRegistryExports& OutExports);
-	static UE_API void GetAssetVariables(const UAnimNextRigVMAssetEditorData* EditorData, TSet<FAnimNextAssetRegistryExportedVariable>& OutExports, TArray<FAnimNextAssetRegistryExportedSharedVariableAsset>& OutExportedSharedVariableAssets, TArray<FAnimNextAssetRegistryExportedSharedVariableStruct>& OutExportedSharedVariableStructs);
-
 	// Gets the exported public functions that are used by a RigVM asset
-	static UE_API void GetAssetFunctions(const UAnimNextRigVMAssetEditorData* InEditorData, FRigVMGraphFunctionHeaderArray& OutExports);
+	static void GetAssetFunctions(const UAnimNextRigVMAssetEditorData* InEditorData, FRigVMGraphFunctionHeaderArray& OutExports);
 
 	// Gets the non-exported private functions that are used by a RigVM asset
 	static UE_API void GetAssetPrivateFunctions(const UAnimNextRigVMAssetEditorData* InEditorData, FRigVMGraphFunctionHeaderArray& OutExports);
 
 	// Gets the asset-registry information needed for representing the contained data into the Workspace Outliner
 	// Note: We pass parent as index to avoid the ref from being invalidated due to realloc as the export array grows recursively
-	static UE_API void GetAssetOutlinerItems(const UAnimNextRigVMAssetEditorData* EditorData, FWorkspaceOutlinerItemExports& OutExports, FAssetRegistryTagsContext Context);
-	static UE_API void CreateSubGraphsOutlinerItemsRecursive(const UAnimNextRigVMAssetEditorData* EditorData, FWorkspaceOutlinerItemExports& OutExports, FWorkspaceOutlinerItemExport& RootExport, int32 ParentExportIndex, URigVMEdGraph* RigVMEdGraph, FAssetRegistryTagsContext Context);
-	static UE_API void CreateFunctionLibraryOutlinerItemsRecursive(const UAnimNextRigVMAssetEditorData* EditorData, FWorkspaceOutlinerItemExports& OutExports, FWorkspaceOutlinerItemExport& RootExport, int32 ParentExportIndex, const TArray<FRigVMGraphFunctionData>& PublicFunctions, const TArray<FRigVMGraphFunctionData>& PrivateFunctions);
-	static UE_API void CreateFunctionsOutlinerItemsRecursive(const UAnimNextRigVMAssetEditorData* EditorData, FWorkspaceOutlinerItemExports& OutExports, FWorkspaceOutlinerItemExport& RootExport, int32 ParentExportIndex, const TArray<FRigVMGraphFunctionData>& Functions, bool bPublicFunctions);
+	static UE_API void GetAssetWorkspaceExports(const UAnimNextRigVMAssetEditorData* EditorData, FWorkspaceOutlinerItemExports& OutExports, FAssetRegistryTagsContext Context);
 
 	// Attempts to determine the type from a variable name & asset/struct path
 	// If the name cannot be found, the returned type will be invalid
@@ -281,8 +197,44 @@ struct FUtils
 	// Make an event name that we use as a wrapper to call RigVM functions 
 	static UE_API FString MakeFunctionWrapperEventName(FName InFunctionName);
 
+	// Returns all unique variable names found within provided editor data, optionally recursing into SharedVariable entries
+	static UE_API void GetVariableNames(UAnimNextRigVMAssetEditorData* InEditorData, TArray<FName>& OutVariableNames, bool bRecursive = true);
+
+	// Return valid variable FName within the context of the provided EditorData (recurses into SharedVariable entries as well) 
+	static UE_API FName GetValidVariableName(UAnimNextRigVMAssetEditorData* InEditorData, FName InBaseName);
+
+	// Return valid variable FName within the context of the provided set of existing variable names
+	static UE_API FName GetValidVariableName(FName InBaseName, const TArrayView<FName> ExistingNames);
+
+	// Deletes the provided variable entry from its editor data, removing any references across the project its assets 
+	static UE_API void DeleteVariable(UAnimNextVariableEntry* VariableEntry, bool bSetupUndoRedo = true, bool bPrintPythonCommands = true);
+
+	// Deletes the provided variable entry from its editor data, renaming any references across the project its assets to reference it by the NewName
+	static UE_API void RenameVariable(UAnimNextVariableEntry* VariableEntry, const FName NewName, bool bSetupUndoRedo = true, bool bPrintPythonCommands = true);
+
+	// Moves the provided variable entry from its editor data to the NewOuter, replacing any references across the project its assets to reference it by its new outer asset
+	static UE_API void MoveVariableToAsset(UAnimNextVariableEntry* VariableEntry, UAnimNextRigVMAsset* NewOuter, bool bSetupUndoRedo = true, bool bPrintPythonCommands = true);
+
+	// Changes the provided variable entry its data type, replacing any references across the project to reference using its new data type
+	static UE_API void SetVariableType(UAnimNextVariableEntry* VariableEntry, const FAnimNextParamType NewType, bool bSetupUndoRedo = true, bool bPrintPythonCommands = true);
+
+	// Replaces the provided variable reference with ReferenceToReplaceWith within the context of InEditorData
+	static UE_API void ReplaceVariableReferences(UAnimNextRigVMAssetEditorData* InEditorData, const FAnimNextSoftVariableReference& ReferenceToFind, const FAnimNextSoftVariableReference& ReferenceToReplaceWith, bool bSetupUndoRedo = true, bool bPrintPythonCommands = true);
+	// Replaces the provided variable reference with ReferenceToReplaceWith within the context of the project
+	static UE_API void ReplaceVariableReferencesAcrossProject(const FAnimNextSoftVariableReference& ReferenceToFind, const FAnimNextSoftVariableReference& ReferenceToReplaceWith, bool bSetupUndoRedo = true, bool bPrintPythonCommands = true);
 private:
 	static UE_API void CompileVariableBindingsInternal(const FRigVMCompileSettings& InSettings, UAnimNextRigVMAsset* InAsset, TArray<URigVMGraph*>& OutGraphs, bool bInThreadSafe);
+	
+	static void GetSubGraphWorkspaceExportsRecursive(const UAnimNextRigVMAssetEditorData* EditorData, FWorkspaceOutlinerItemExports& OutExports, const FWorkspaceOutlinerItemExport& RootExport, int32 ParentExportIndex, URigVMEdGraph* RigVMEdGraph, FAssetRegistryTagsContext Context);
+	static void GetFunctionLibraryWorkspaceExportsRecursive(const UAnimNextRigVMAssetEditorData* EditorData, FWorkspaceOutlinerItemExports& OutExports, const FWorkspaceOutlinerItemExport& RootExport, int32 ParentExportIndex, const TArray<FRigVMGraphFunctionData>& PublicFunctions, const TArray<FRigVMGraphFunctionData>& PrivateFunctions);
+	static void GetFunctionsWorkspaceExportsRecursive(const UAnimNextRigVMAssetEditorData* EditorData, FWorkspaceOutlinerItemExports& OutExports, const FWorkspaceOutlinerItemExport& RootExport, int32 ParentExportIndex, const TArray<FRigVMGraphFunctionData>& Functions, bool bPublicFunctions);
+	
+	// Gets the exported variables that are used by a RigVM asset	
+	static void GetAssetVariableExports(const UAnimNextRigVMAssetEditorData* EditorData, TSet<FAnimNextExport>& OutExports, FAssetRegistryTagsContext Context);
+
+	static void GetStructVariableExports(const UScriptStruct* Struct, TSet<FAnimNextExport>& OutExports);
+	static void GetSubGraphVariableExportsRecursive(const URigVMEdGraph* RigVMEdGraph, TSet<FAnimNextExport>& OutExports);
+	static void ProcessPinAssetReferences(const URigVMPin* InPin, FWorkspaceOutlinerItemExports& OutExports, const FWorkspaceOutlinerItemExport& RootExport, int32 ParentExportIndex);
 };
 
 }
diff --git a/Engine/Plugins/Experimental/AnimNext/Source/AnimNextUncookedOnly/Private/AnimNextAssetWorkspaceAssetUserData.cpp b/Engine/Plugins/Experimental/AnimNext/Source/AnimNextUncookedOnly/Private/AnimNextAssetWorkspaceAssetUserData.cpp
--- a/Engine/Plugins/Experimental/AnimNext/Source/AnimNextUncookedOnly/Private/AnimNextAssetWorkspaceAssetUserData.cpp
+++ b/Engine/Plugins/Experimental/AnimNext/Source/AnimNextUncookedOnly/Private/AnimNextAssetWorkspaceAssetUserData.cpp
@@ -2,6 +2,7 @@
 
 #include "AnimNextAssetWorkspaceAssetUserData.h"
 
+#include "AnimNextExports.h"
 #include "UncookedOnlyUtils.h"
 #include "Variables/AnimNextSharedVariables.h"
 #include "Module/AnimNextModule.h"
@@ -18,7 +19,7 @@ void UAnimNextAssetWorkspaceAssetUserData::GetAssetRegistryTags(FAssetRegistryTa
 		CachedExports.Exports.Reset();
 		{
 			UAnimNextRigVMAsset* Asset = CastChecked<UAnimNextRigVMAsset>(GetOuter());
-			const UAnimNextRigVMAssetEditorData* GraphEditorData = UE::AnimNext::UncookedOnly::FUtils::GetEditorData<UAnimNextRigVMAssetEditorData>(Asset);
+			if (const UAnimNextRigVMAssetEditorData* GraphEditorData = UE::AnimNext::UncookedOnly::FUtils::GetEditorData<UAnimNextRigVMAssetEditorData>(Asset))
 			{
 				FWorkspaceOutlinerItemExport& RootAssetExport = CachedExports.Exports.Add_GetRef(FWorkspaceOutlinerItemExport(Asset->GetFName(), Asset));
 
@@ -28,7 +29,7 @@ void UAnimNextAssetWorkspaceAssetUserData::GetAssetRegistryTags(FAssetRegistryTa
 				}
 				else if(Asset->IsA<UAnimNextSharedVariables>())
 				{
-					RootAssetExport.GetData().InitializeAsScriptStruct(FAnimNextDataInterfaceOutlinerData::StaticStruct());
+					RootAssetExport.GetData().InitializeAsScriptStruct(FAnimNextSharedVariablesOutlinerData::StaticStruct());
 				}
 				else
 				{
@@ -36,9 +37,9 @@ void UAnimNextAssetWorkspaceAssetUserData::GetAssetRegistryTags(FAssetRegistryTa
 				}
 				FAnimNextRigVMAssetOutlinerData& Data = RootAssetExport.GetData().GetMutable<FAnimNextRigVMAssetOutlinerData>();
 				Data.SoftAssetPtr = Asset;
+
+				UE::AnimNext::UncookedOnly::FUtils::GetAssetWorkspaceExports(GraphEditorData, CachedExports, Context);
 			}
-		
-			UE::AnimNext::UncookedOnly::FUtils::GetAssetOutlinerItems(GraphEditorData, CachedExports, Context);
 		}
 	}
 	
diff --git a/Engine/Plugins/Experimental/AnimNext/Source/AnimNextUncookedOnly/Private/AnimNextControllerBase.cpp b/Engine/Plugins/Experimental/AnimNext/Source/AnimNextUncookedOnly/Private/AnimNextControllerBase.cpp
--- a/Engine/Plugins/Experimental/AnimNext/Source/AnimNextUncookedOnly/Private/AnimNextControllerBase.cpp
+++ b/Engine/Plugins/Experimental/AnimNext/Source/AnimNextUncookedOnly/Private/AnimNextControllerBase.cpp
@@ -1,7 +1,8 @@
-﻿// Copyright Epic Games, Inc. All Rights Reserved.
+// Copyright Epic Games, Inc. All Rights Reserved.
 
 #include "AnimNextControllerBase.h"
 #include "AnimNextSharedVariableNode.h"
+#include "RigVMPythonUtils.h"
 #include "Variables/AnimNextSharedVariables.h"
 
 UAnimNextSharedVariableNode* UAnimNextControllerBase::AddAssetSharedVariableNode(const UAnimNextSharedVariables* InAsset, const FName& InVariableName, const FString& InCPPType, UObject* InCPPTypeObject, bool bIsGetter, const FString& InDefaultValue, const FVector2D& InPosition, const FString& InNodeName, bool bSetupUndoRedo, bool bPrintPythonCommand)
@@ -12,7 +13,7 @@ UAnimNextSharedVariableNode* UAnimNextControllerBase::AddAssetSharedVariableNode
 		return nullptr;
 	}
 
-	return AddSharedVariableNode(InAsset->GetPathName(), InVariableName, InCPPType, InCPPTypeObject, bIsGetter, InDefaultValue, InPosition, InNodeName, bSetupUndoRedo, bPrintPythonCommand);
+	return AddSharedVariableNode(InAsset->GetPathName(), InVariableName, InCPPType, InCPPTypeObject->GetPathName(), bIsGetter, InDefaultValue, InPosition, InNodeName, bSetupUndoRedo, bPrintPythonCommand);
 }
 
 UAnimNextSharedVariableNode* UAnimNextControllerBase::AddStructSharedVariableNode(const UScriptStruct* InStruct, const FName& InVariableName, const FString& InCPPType, UObject* InCPPTypeObject, bool bIsGetter, const FString& InDefaultValue, const FVector2D& InPosition, const FString& InNodeName, bool bSetupUndoRedo, bool bPrintPythonCommand)
@@ -23,10 +24,10 @@ UAnimNextSharedVariableNode* UAnimNextControllerBase::AddStructSharedVariableNod
 		return nullptr;
 	}
 
-	return AddSharedVariableNode(InStruct->GetPathName(), InVariableName, InCPPType, InCPPTypeObject, bIsGetter, InDefaultValue, InPosition, InNodeName, bSetupUndoRedo, bPrintPythonCommand);
+	return AddSharedVariableNode(InStruct->GetPathName(), InVariableName, InCPPType, InCPPTypeObject->GetPathName(), bIsGetter, InDefaultValue, InPosition, InNodeName, bSetupUndoRedo, bPrintPythonCommand);
 }
 
-UAnimNextSharedVariableNode* UAnimNextControllerBase::AddSharedVariableNode(const FString& InSourceObjectPath, const FName& InVariableName, const FString& InCPPType, UObject* InCPPTypeObject, bool bIsGetter, const FString& InDefaultValue, const FVector2D& InPosition, const FString& InNodeName, bool bSetupUndoRedo, bool bPrintPythonCommand)
+UAnimNextSharedVariableNode* UAnimNextControllerBase::AddSharedVariableNode(const FString& InSourceObjectPath, const FName& InVariableName, const FString& InCPPType, const FString& InCPPTypeObjectPath, bool bIsGetter, const FString& InDefaultValue, const FVector2D& InPosition, const FString& InNodeName, bool bSetupUndoRedo, bool bPrintPythonCommand)
 {
 	UObject* SourceObject = nullptr;
 	if (!InSourceObjectPath.IsEmpty())
@@ -34,11 +35,23 @@ UAnimNextSharedVariableNode* UAnimNextControllerBase::AddSharedVariableNode(cons
 		SourceObject = RigVMTypeUtils::FindObjectFromCPPTypeObjectPath<UObject>(InSourceObjectPath);
 		if (SourceObject == nullptr)
 		{
+			ReportErrorf(TEXT("Cannot find cpp type object for path '%s'."), *InSourceObjectPath);
 			return nullptr;
 		}
 	}
 
-	UAnimNextSharedVariableNode* VariableNode = CastChecked<UAnimNextSharedVariableNode>(AddVariableNode(InVariableName, UAnimNextSharedVariableNode::StaticClass(), InCPPType, InCPPTypeObject, bIsGetter, InDefaultValue, InPosition, InNodeName, bSetupUndoRedo, bPrintPythonCommand), ECastCheckedType::Type::NullAllowed);
+	UObject* CPPTypeObject = nullptr;
+	if (!InCPPTypeObjectPath.IsEmpty())
+	{
+		CPPTypeObject = RigVMTypeUtils::FindObjectFromCPPTypeObjectPath<UObject>(InCPPTypeObjectPath);
+		if (CPPTypeObject == nullptr)
+		{
+			ReportErrorf(TEXT("Cannot find cpp type object for path '%s'."), *InCPPTypeObjectPath);
+			return nullptr;
+		}
+	}
+
+	UAnimNextSharedVariableNode* VariableNode = CastChecked<UAnimNextSharedVariableNode>(AddVariableNode(InVariableName, UAnimNextSharedVariableNode::StaticClass(), InCPPType, CPPTypeObject, bIsGetter, InDefaultValue, InPosition, InNodeName, bSetupUndoRedo, false), ECastCheckedType::Type::NullAllowed);
 	if (VariableNode == nullptr)
 	{
 		return nullptr;
@@ -59,28 +72,122 @@ UAnimNextSharedVariableNode* UAnimNextControllerBase::AddSharedVariableNode(cons
 		return nullptr;
 	}
 
+	if (bPrintPythonCommand)
+	{
+		const FString GraphName = GetSchema()->GetSanitizedGraphName(GetGraph()->GetGraphName());
+
+		RigVMPythonUtils::Print(GetSchema()->GetGraphOuterName(GetGraph()), 
+			FString::Printf(TEXT("asset.get_controller_by_name('%s').add_shared_variable_node('%s', '%s', '%s', '%s', '%s', '%s', %s, %s)"),
+				*GraphName,
+				*InSourceObjectPath,
+				*InVariableName.ToString(),
+				*InCPPType,
+				*InCPPTypeObjectPath,
+				bIsGetter ? TEXT("True") : TEXT("False"),
+				*InDefaultValue,
+				*RigVMPythonUtils::Vector2DToPythonString(InPosition),
+				*InNodeName
+			)
+		);
+	}
+
 	return VariableNode;
 }
 
-void UAnimNextControllerBase::ReplaceVariableNodeWithSharedVariableNode(URigVMVariableNode* InVariableNode, FName InNewVariableName, const UObject* InAssetOrStruct)
+void UAnimNextControllerBase::RefreshSharedVariableNode(const FName& InNodeName, const FString& InSourceObjectPath, const FName& InVariableName, const FString& InCPPType, UObject* InCPPTypeObject, bool bSetupUndoRedo, bool bPrintPythonCommand)
+{
+	UObject* SourceObject = nullptr;
+	if (!InSourceObjectPath.IsEmpty())
+	{
+		SourceObject = RigVMTypeUtils::FindObjectFromCPPTypeObjectPath<UObject>(InSourceObjectPath);
+	}
+
+	URigVMGraph* Graph = GetGraph();
+	check(Graph);
+	if (UAnimNextSharedVariableNode* VariableNode = Cast<UAnimNextSharedVariableNode>(Graph->FindNodeByName(InNodeName)))
+	{
+		if (bSetupUndoRedo)
+		{
+			VariableNode->Modify();
+		}
+		
+		constexpr bool bSetupOrphanPins = false;
+		RefreshVariableNode(InNodeName, InVariableName, InCPPType, InCPPTypeObject, bSetupUndoRedo, bSetupOrphanPins);
+
+		if (SourceObject)
+		{
+			if (const UAnimNextSharedVariables* Asset = Cast<UAnimNextSharedVariables>(SourceObject))
+			{
+				VariableNode->Type = EAnimNextSharedVariablesType::Asset;
+				VariableNode->Asset = Asset;
+			}
+			else if (const UScriptStruct* Struct = Cast<UScriptStruct>(SourceObject))
+			{
+				VariableNode->Type = EAnimNextSharedVariablesType::Struct;
+				VariableNode->Struct = Struct;
+			}
+		}
+		else
+		{
+			VariableNode->Type = EAnimNextSharedVariablesType::Asset;
+			VariableNode->Asset = nullptr;
+		}
+
+		Notify(ERigVMGraphNotifType::NodeTitleChanged, VariableNode);
+
+		if (bPrintPythonCommand)
+		{
+			
+		}
+	}
+}
+
+UAnimNextSharedVariableNode* UAnimNextControllerBase::ReplaceVariableNodeWithSharedVariableNode(URigVMVariableNode* InVariableNode, FName InNewVariableName, const UObject* InAssetOrStruct, bool bSetupUndoRedo, bool bPrintPythonCommand)
 {
 	check(InVariableNode);
 	check(InAssetOrStruct);
 	check(InNewVariableName != NAME_None);
 
-	UAnimNextSharedVariableNode* SharedVariableNode = AddSharedVariableNode(InAssetOrStruct->GetPathName(), InNewVariableName, InVariableNode->GetCPPType(), InVariableNode->GetCPPTypeObject(), InVariableNode->IsGetter(), InVariableNode->GetDefaultValue(), InVariableNode->GetPosition(), FString(), false, false);
+	UAnimNextSharedVariableNode* SharedVariableNode = AddSharedVariableNode(InAssetOrStruct->GetPathName(), InNewVariableName, InVariableNode->GetCPPType(), InVariableNode->GetCPPTypeObject() ? InVariableNode->GetCPPTypeObject()->GetPathName() : TEXT(""), InVariableNode->IsGetter(), InVariableNode->GetDefaultValue(), InVariableNode->GetPosition(), FString(), bSetupUndoRedo, bPrintPythonCommand);
 	if (SharedVariableNode)
 	{
 		if (!InVariableNode->IsGetter())
 		{
-			RewireLinks(InVariableNode->GetValuePin(), SharedVariableNode->GetValuePin(), true, false);
-			RewireLinks(InVariableNode->FindExecutePin(), SharedVariableNode->FindExecutePin(), true, false);
-			RewireLinks(InVariableNode->FindExecutePin(), SharedVariableNode->FindExecutePin(), false, false);
+			RewireLinks(InVariableNode->GetValuePin(), SharedVariableNode->GetValuePin(), true, bSetupUndoRedo);
+			RewireLinks(InVariableNode->FindExecutePin(), SharedVariableNode->FindExecutePin(), true, bSetupUndoRedo);
+			RewireLinks(InVariableNode->FindExecutePin(), SharedVariableNode->FindExecutePin(), false, bSetupUndoRedo);
 		}
 		else
 		{
-			RewireLinks(InVariableNode->GetValuePin(), SharedVariableNode->GetValuePin(), false, false);
+			RewireLinks(InVariableNode->GetValuePin(), SharedVariableNode->GetValuePin(), false, bSetupUndoRedo);
 		}
-		RemoveNode(InVariableNode, false, false);
+		RemoveNode(InVariableNode, bSetupUndoRedo, bPrintPythonCommand);
 	}
+
+	return SharedVariableNode;
+}
+
+URigVMVariableNode* UAnimNextControllerBase::ReplaceSharedVariableNodeWithVariableNode(UAnimNextSharedVariableNode* InVariableNode, FName InNewVariableName, bool bSetupUndoRedo, bool bPrintPythonCommand)
+{
+	check(InVariableNode);
+	check(InNewVariableName != NAME_None);
+
+
+	URigVMVariableNode* VariableNode = AddVariableNode(InNewVariableName, InVariableNode->GetCPPType(), InVariableNode->GetCPPTypeObject(), InVariableNode->IsGetter(), InVariableNode->GetDefaultValue(), InVariableNode->GetPosition(), FString(), bSetupUndoRedo, bPrintPythonCommand);
+	if (VariableNode)
+	{
+		if (!InVariableNode->IsGetter())
+		{
+			RewireLinks(InVariableNode->GetValuePin(), VariableNode->GetValuePin(), true, bSetupUndoRedo);
+			RewireLinks(InVariableNode->FindExecutePin(), VariableNode->FindExecutePin(), true, bSetupUndoRedo);
+			RewireLinks(InVariableNode->FindExecutePin(), VariableNode->FindExecutePin(), false, bSetupUndoRedo);
+		}
+		else
+		{
+			RewireLinks(InVariableNode->GetValuePin(), VariableNode->GetValuePin(), false, bSetupUndoRedo);
+		}
+		RemoveNode(InVariableNode, bSetupUndoRedo, bPrintPythonCommand);
+	}
+
+	return VariableNode;
 }
\ No newline at end of file
diff --git a/Engine/Plugins/Experimental/AnimNext/Source/AnimNextUncookedOnly/Private/AnimNextRigVMAssetEditorData.cpp b/Engine/Plugins/Experimental/AnimNext/Source/AnimNextUncookedOnly/Private/AnimNextRigVMAssetEditorData.cpp
--- a/Engine/Plugins/Experimental/AnimNext/Source/AnimNextUncookedOnly/Private/AnimNextRigVMAssetEditorData.cpp
+++ b/Engine/Plugins/Experimental/AnimNext/Source/AnimNextUncookedOnly/Private/AnimNextRigVMAssetEditorData.cpp
@@ -21,7 +21,6 @@
 #include "IWorkspaceEditorModule.h"
 #include "ObjectTools.h"
 #include "UncookedOnlyUtils.h"
-#include "Animation/AnimCompressionTypes.h"
 #include "Animation/Skeleton.h"
 #include "AnimNext/Private/AnimNextRigVMFunctionData.h"
 #include "AnimNext/Private/Variables/StructDataCache.h"
@@ -33,13 +32,11 @@
 #include "Graph/RigUnit_AnimNextBeginExecution.h"
 #include "Misc/TransactionObjectEvent.h"
 #include "Module/AnimNextEventGraphSchema.h"
-#include "Module/RigUnit_AnimNextModuleEvents.h"
 #include "RigVMModel/RigVMFunctionLibrary.h"
 #include "RigVMModel/RigVMNotifications.h"
 #include "RigVMModel/Nodes/RigVMAggregateNode.h"
 #include "RigVMModel/Nodes/RigVMCollapseNode.h"
 #include "UObject/AssetRegistryTagsContext.h"
-#include "Variables/AnimNextSharedVariables_EditorData.h"
 #include "RigVMFunctions/Execution/RigVMFunction_UserDefinedEvent.h"
 #include "Misc/UObjectToken.h"
 #include "UObject/SavePackage.h"
@@ -308,9 +305,7 @@ void UAnimNextRigVMAssetEditorData::GetAssetRegistryTags(FAssetRegistryTagsConte
 		{
 			CachedExports = FAnimNextAssetRegistryExports();
 			FAnimNextAssetRegistryExports& OutExports = CachedExports.GetValue();
-
 			GetAnimNextAssetRegistryTags(Context, OutExports);
-			UE::AnimNext::UncookedOnly::FUtils::GetAssetVariables(this, OutExports);
 		}
 
 		FString TagValue;
@@ -376,6 +371,11 @@ void UAnimNextRigVMAssetEditorData::HandlePackageDone()
 	RecompileVM();
 }
 
+void UAnimNextRigVMAssetEditorData::GetAnimNextAssetRegistryTags(FAssetRegistryTagsContext& Context, FAnimNextAssetRegistryExports& OutExports) const
+{	
+	UE::AnimNext::UncookedOnly::FUtils::GetAssetVariableExports(this, OutExports, Context);
+}
+
 void UAnimNextRigVMAssetEditorData::RefreshAllModels(ERigVMLoadType InLoadType)
 {
 }
@@ -509,9 +509,6 @@ void UAnimNextRigVMAssetEditorData::RenameGraph(const FString& InNodePath, const
 	{
 		if (UEdGraph* EdGraph = Cast<UEdGraph>(GetEditorObjectForRigVMGraph(ModelForNodePath)))
 		{
-			FName OldName = NAME_None;
-			OldName = EdGraph->GetFName();
-
 			RigVMClient.RenameModel(InNodePath, InNewName, true);
 		}
 	}
@@ -1092,9 +1089,6 @@ void UAnimNextRigVMAssetEditorData::RecompileVM()
 			RigVMCompiledEvent.Broadcast(Asset, Asset->VM, Asset->ExtendedExecuteContext);
 		}
 
-		FAnimNextAssetRegistryExports Exports;
-		FUtils::GetAssetVariables(this, Exports);
-
 #if WITH_EDITOR
 		// Display programmatic graphs
 		if(CVarDumpProgrammaticGraphs.GetValueOnGameThread())
@@ -1481,6 +1475,171 @@ UAnimNextRigVMAssetEntry* UAnimNextRigVMAssetEditorData::FindEntry(FName InName)
 	return FoundEntry != nullptr ? *FoundEntry : nullptr;
 }
 
+bool UAnimNextRigVMAssetEditorData::AddCategory(const FString& CategoryName, bool bInSetupUndoRedo, bool bPrintPythonCommand)
+{
+	if (CategoryName.IsEmpty())
+    {
+		ReportError(*FString::Printf(TEXT("UAnimNextRigVMAssetEditorData::AddCategory: Invalid empty category name provided.")));
+		return false;
+    }
+
+
+	const int32 CategoryIndex = VariableAndFunctionCategories.IndexOfByKey(CategoryName);
+	if(CategoryIndex != INDEX_NONE)
+	{
+		ReportError(*FString::Printf(TEXT("UAnimNextRigVMAssetEditorData::AddCategory: Already contains category with name %s."), *CategoryName));
+		return false;
+	}
+
+	if (bInSetupUndoRedo)
+	{
+		Modify();
+	}
+
+	VariableAndFunctionCategories.Add(CategoryName);
+	BroadcastModified(EAnimNextEditorDataNotifType::CategoryAdded, this);
+
+	if (bPrintPythonCommand)
+	{
+		RigVMPythonUtils::Print(GetName(),
+				FString::Printf(TEXT("asset.add_category('%s')"),
+				 *CategoryName));
+	}
+
+	return true;
+}
+
+bool UAnimNextRigVMAssetEditorData::RenameCategory(const FString& OldName, const FString& NewName, bool bInSetupUndoRedo, bool bPrintPythonCommand)
+{
+	const int32 CategoryIndex = VariableAndFunctionCategories.IndexOfByKey(OldName);
+	if (CategoryIndex == INDEX_NONE)
+	{
+		ReportError(*FString::Printf(TEXT("UAnimNextRigVMAssetEditorData::RenameCategory: Invalid existing category name provided (not found).")));
+		return false;
+	}
+	
+	if (bInSetupUndoRedo)
+	{
+		Modify();
+	}
+
+	VariableAndFunctionCategories[CategoryIndex] = NewName;
+
+	for(UAnimNextRigVMAssetEntry* Entry : Entries)
+	{
+		if(UAnimNextVariableEntry* VariableEntry = Cast<UAnimNextVariableEntry>(Entry))
+		{
+			if (VariableEntry->GetVariableCategory() == OldName)
+			{
+				VariableEntry->SetVariableCategory(NewName, bInSetupUndoRedo);
+			}
+		}
+	}
+	
+	BroadcastModified(EAnimNextEditorDataNotifType::CategoryChanged, this);
+
+	if (bPrintPythonCommand)
+	{
+		RigVMPythonUtils::Print(GetName(),
+			FString::Printf(TEXT("asset.rename_category('%s', '%s')"),
+			 *OldName,
+			 *NewName));
+	}
+
+	return true;
+}
+
+void UAnimNextRigVMAssetEditorData::ReorderCategory(const FString& CategoryName, const FString& BeforeCategoryName, bool bInSetupUndoRedo,
+	bool bPrintPythonCommand)
+{
+	const int32 CategoryIndex = VariableAndFunctionCategories.IndexOfByKey(CategoryName);
+	if (CategoryIndex == INDEX_NONE)
+	{
+		return;
+	}
+	
+	const int32 BeforeCategoryIndex = VariableAndFunctionCategories.IndexOfByKey(BeforeCategoryName);
+	if (BeforeCategoryIndex == INDEX_NONE)
+	{
+		return;
+	}
+
+	if (BeforeCategoryIndex == CategoryIndex)
+	{
+		return;
+	}
+
+	if (bInSetupUndoRedo)
+	{
+		Modify();
+	}
+
+	VariableAndFunctionCategories.RemoveAt(CategoryIndex);
+	
+	int32 NewCategoryIndex = 0;
+
+	// Moving up
+	if (CategoryIndex > BeforeCategoryIndex)
+	{
+		// RemoveAt didn't have any impact
+		NewCategoryIndex = FMath::Max(BeforeCategoryIndex, 0);
+	}
+	else
+	{
+		// RemoveAt means insertion index has to be offset by 1
+		NewCategoryIndex = FMath::Max(BeforeCategoryIndex - 1, 0);
+	}
+	
+	VariableAndFunctionCategories.Insert(CategoryName, NewCategoryIndex);
+
+	BroadcastModified(EAnimNextEditorDataNotifType::CategoryChanged, this);
+}
+
+void UAnimNextRigVMAssetEditorData::ReorderVariable(UAnimNextVariableEntry* VariableEntry, const UAnimNextVariableEntry* BeforeVariableEntry, bool bSetupUndoRedo, bool bPrintPythonCommand)
+{
+	const int32 VariableEntryIndex = InternalEntries.IndexOfByKey(VariableEntry);
+	if (VariableEntryIndex == INDEX_NONE)
+	{
+		return;
+	}
+	
+	const int32 BeforeVariableEntryIndex = InternalEntries.IndexOfByKey(BeforeVariableEntry);
+	if (BeforeVariableEntryIndex == INDEX_NONE)
+	{
+		return;
+	}
+
+	if (BeforeVariableEntryIndex == VariableEntryIndex)
+	{
+		return;
+	}
+
+	if (bSetupUndoRedo)
+	{
+		Modify();
+	}
+
+	RemoveEntryInternal(VariableEntry);
+
+	int32 NewVariableEntryIndex = 0;
+
+	// Moving up
+	if (VariableEntryIndex > BeforeVariableEntryIndex)
+	{
+		// RemoveAt didn't have any impact
+		NewVariableEntryIndex = FMath::Max(BeforeVariableEntryIndex, 0);
+	}
+	else
+	{
+		// RemoveAt means insertion index has to be offset by 1
+		NewVariableEntryIndex = FMath::Max(BeforeVariableEntryIndex - 1, 0);
+	}
+	
+	InsertEntryInternal(VariableEntry, NewVariableEntryIndex);
+
+	BroadcastModified(EAnimNextEditorDataNotifType::EntryAdded, VariableEntry);
+}
+
 bool UAnimNextRigVMAssetLibrary::RemoveEntry(UAnimNextRigVMAsset* InAsset, UAnimNextRigVMAssetEntry* InEntry,  bool bSetupUndoRedo, bool bPrintPythonCommand)
 {
 	return UE::AnimNext::UncookedOnly::FUtils::GetEditorData(InAsset)->RemoveEntry(InEntry, bSetupUndoRedo, bPrintPythonCommand);
@@ -2281,6 +2440,16 @@ URigVMLibraryNode* UAnimNextRigVMAssetLibrary::AddFunction(UAnimNextRigVMAsset*
 	return UE::AnimNext::UncookedOnly::FUtils::GetEditorData(InAsset)->AddFunction(InFunctionName, bInMutable, bSetupUndoRedo, bPrintPythonCommand);
 }
 
+bool UAnimNextRigVMAssetLibrary::AddCategory(UAnimNextRigVMAsset* InAsset, const FString& CategoryName, bool bSetupUndoRedo, bool bPrintPythonCommand)
+{
+	return UE::AnimNext::UncookedOnly::FUtils::GetEditorData(InAsset)->AddCategory(CategoryName, bSetupUndoRedo, bPrintPythonCommand);
+}
+
+bool UAnimNextRigVMAssetLibrary::RenameCategory(UAnimNextRigVMAsset* InAsset, const FString& CategoryName, const FString& NewCategoryName, bool bSetupUndoRedo, bool bPrintPythonCommand)
+{
+	return UE::AnimNext::UncookedOnly::FUtils::GetEditorData(InAsset)->RenameCategory(CategoryName, NewCategoryName, bSetupUndoRedo, bPrintPythonCommand);
+}
+
 URigVMLibraryNode* UAnimNextRigVMAssetEditorData::AddFunction(FName InFunctionName, bool bInMutable, bool bInSetupUndoRedo, bool bPrintPythonCommand)
 {
 	URigVMController* Controller = RigVMClient.GetOrCreateController(GetLocalFunctionLibrary());
@@ -2536,15 +2705,20 @@ void UAnimNextRigVMAssetEditorData::AddReferencedObjects(UObject* InThis, FRefer
 
 void UAnimNextRigVMAssetEditorData::AddEntryInternal(UAnimNextRigVMAssetEntry* InEntry)
 {
-	// If we are using external packages, dont persist this entry
+	InsertEntryInternal(InEntry, FMath::Max(Entries.Num() - 1, 0));	
+}
+
+void UAnimNextRigVMAssetEditorData::InsertEntryInternal(UAnimNextRigVMAssetEntry* InEntry, int32 InsertionIndex)
+{
+	// If we are using external packages, do not persist this entry
 	if(bUsesExternalPackages)
 	{
-		Entries.Add(InEntry);
+		Entries.Insert(InEntry, InsertionIndex);
 	}
 	else
 	{
-		InternalEntries.Add(InEntry);
-		Entries.Add(InEntry);
+		InternalEntries.Insert(InEntry, InsertionIndex);
+		Entries.Insert(InEntry, InsertionIndex);
 	}
 }
 
@@ -2902,7 +3076,7 @@ void UAnimNextRigVMAssetEditorData::UpgradeDataInterfaces()
 
 					if (AssetOrStruct)
 					{
-						Controller->ReplaceVariableNodeWithSharedVariableNode(VariableNode, NewVariableName, AssetOrStruct);
+						Controller->ReplaceVariableNodeWithSharedVariableNode(VariableNode, NewVariableName, AssetOrStruct, false, false);
 						break;
 					}
 				}
diff --git a/Engine/Plugins/Experimental/AnimNext/Source/AnimNextUncookedOnly/Private/Entries/AnimNextVariableEntry.cpp b/Engine/Plugins/Experimental/AnimNext/Source/AnimNextUncookedOnly/Private/Entries/AnimNextVariableEntry.cpp
--- a/Engine/Plugins/Experimental/AnimNext/Source/AnimNextUncookedOnly/Private/Entries/AnimNextVariableEntry.cpp
+++ b/Engine/Plugins/Experimental/AnimNext/Source/AnimNextUncookedOnly/Private/Entries/AnimNextVariableEntry.cpp
@@ -12,8 +12,6 @@ const FLazyName IAnimNextRigVMVariableInterface::ValueName("Value");
 
 UAnimNextVariableEntry::UAnimNextVariableEntry()
 {
-	// TEMP: This should be selectable
-	Binding.BindingData.InitializeAs<FAnimNextUniversalObjectLocatorBindingData>();
 }
 
 FAnimNextParamType UAnimNextVariableEntry::GetExportType() const
@@ -140,7 +138,7 @@ bool UAnimNextVariableEntry::GetDefaultValue(const FProperty*& OutProperty, TCon
 	const FPropertyBagPropertyDesc* Desc = DefaultValue.FindPropertyDescByName(IAnimNextRigVMVariableInterface::ValueName);
 	if(Desc == nullptr)
 	{
-		UE_LOGFMT(LogAnimation, Error, "UAnimNextVariableEntry::SetDefaultValue: Could not find default value in property bag");
+		UE_LOGFMT(LogAnimation, Error, "UAnimNextVariableEntry::GetDefaultValue: Could not find default value in property bag");
 		return false;
 	}
 
@@ -151,6 +149,21 @@ bool UAnimNextVariableEntry::GetDefaultValue(const FProperty*& OutProperty, TCon
 	return true;
 }
 
+bool UAnimNextVariableEntry::GetDefaultValueString(FString& OutValueString) const
+{
+	const FPropertyBagPropertyDesc* Desc = DefaultValue.FindPropertyDescByName(IAnimNextRigVMVariableInterface::ValueName);
+	if(Desc == nullptr)
+	{
+		UE_LOGFMT(LogAnimation, Error, "UAnimNextVariableEntry::GetDefaultValueString: Could not find default value in property bag");
+		return false;
+	}
+
+	check(Desc->CachedProperty);
+	Desc->CachedProperty->ExportTextItem_InContainer(OutValueString, DefaultValue.GetValue().GetMemory(), nullptr, nullptr, 0);
+
+	return true;
+}
+
 void UAnimNextVariableEntry::SetBindingType(UScriptStruct* InBindingTypeStruct, bool bSetupUndoRedo)
 {
 	check(InBindingTypeStruct == nullptr || InBindingTypeStruct->IsChildOf(FAnimNextVariableBindingData::StaticStruct()));
@@ -189,6 +202,23 @@ TConstStructView<FAnimNextVariableBindingData> UAnimNextVariableEntry::GetBindin
 	return Binding.BindingData;
 }
 
+FStringView UAnimNextVariableEntry::GetVariableCategory() const
+{
+	return Category;
+}
+
+void UAnimNextVariableEntry::SetVariableCategory(FStringView InCategoryName, bool bSetupUndoRedo)
+{
+	if(bSetupUndoRedo)
+	{
+		Modify();
+	}
+
+	Category = InCategoryName;
+	
+	BroadcastModified(EAnimNextEditorDataNotifType::VariableCategoryChanged);
+}
+
 void UAnimNextVariableEntry::SetEntryName(FName InName, bool bSetupUndoRedo)
 {
 	if(bSetupUndoRedo)
