AnimNext variables UI/UX changes:

diff --git a/Engine/Plugins/Experimental/AnimNext/Source/AnimNextUncookedOnly/Private/UncookedOnlyUtils.cpp b/Engine/Plugins/Experimental/AnimNext/Source/AnimNextUncookedOnly/Private/UncookedOnlyUtils.cpp
--- a/Engine/Plugins/Experimental/AnimNext/Source/AnimNextUncookedOnly/Private/UncookedOnlyUtils.cpp
+++ b/Engine/Plugins/Experimental/AnimNext/Source/AnimNextUncookedOnly/Private/UncookedOnlyUtils.cpp
@@ -32,6 +32,10 @@
 #include "RigVMFunctions/Execution/RigVMFunction_UserDefinedEvent.h"
 #include "UObject/AssetRegistryTagsContext.h"
 #include "Variables/AnimNextSoftVariableReference.h"
+#include "AnimNextExports.h"
+#include "AnimNextScopedCompileJob.h"
+#include "AnimNextSharedVariableNode.h"
+#include "Styling/SlateIconFinder.h"
 
 #define LOCTEXT_NAMESPACE "AnimNextUncookedOnlyUtils"
 
@@ -720,16 +724,19 @@ FAnimNextParamType FUtils::FindVariableType(const FAnimNextVariableReference& In
 			return FAnimNextParamType::FromProperty(Property);
 		}
 	}
-	else if (const UAnimNextRigVMAsset* Asset =  Cast<UAnimNextRigVMAsset>(InVariableReference.GetObject()))
+	else if (const UAnimNextRigVMAsset* Asset = Cast<UAnimNextRigVMAsset>(InVariableReference.GetObject()))
 	{
-		const UAnimNextRigVMAssetEditorData* EditorData = GetEditorData<UAnimNextRigVMAssetEditorData>(Asset);
+		// Query the asset registry for an asset's variable exports
+		FAssetData AssetData(Asset);
 		FAnimNextAssetRegistryExports Exports;
-		GetAssetVariables(EditorData, Exports);
-		for(const FAnimNextAssetRegistryExportedVariable& Variable : Exports.Variables)
+		if (GetExportedVariablesForAsset(AssetData, Exports))
 		{
-			if(Variable.Name == InVariableReference.GetName())
+			for(const FAnimNextExport& Export : Exports.Exports)
 			{
-				return Variable.Type;
+				if(Export.Identifier == InVariableReference.GetName())
+				{
+					return Export.Data.Get<FAnimNextVariableDeclarationData>().Type;
+				}
 			}
 		}
 	}
@@ -761,12 +768,16 @@ FAnimNextParamType FUtils::FindVariableType(const FAnimNextSoftVariableReference
 		{
 			const UAnimNextRigVMAssetEditorData* EditorData = GetEditorData<UAnimNextRigVMAssetEditorData>(Asset);
 			FAnimNextAssetRegistryExports Exports;
-			GetAssetVariables(EditorData, Exports);
-			for(const FAnimNextAssetRegistryExportedVariable& Variable : Exports.Variables)
+
+			FAssetRegistryTagsContextData ContextData(EditorData, EAssetRegistryTagsCaller::Uncategorized);
+			const FAssetRegistryTagsContext Context(ContextData);
+			GetAssetVariableExports(EditorData, Exports, Context);
+
+			for(const FAnimNextExport& Export : Exports.Exports)
 			{
-				if(Variable.Name == InSoftVariableReference.GetName())
+				if(Export.Identifier == InSoftVariableReference.GetName())
 				{
-					return Variable.Type;
+					return Export.Data.Get<FAnimNextVariableDeclarationData>().Type;
 				}
 			}
 		}
@@ -776,11 +787,11 @@ FAnimNextParamType FUtils::FindVariableType(const FAnimNextSoftVariableReference
 			FAnimNextAssetRegistryExports Exports;
 			if (GetExportedVariablesForAsset(AssetData, Exports))
 			{
-				for(const FAnimNextAssetRegistryExportedVariable& Variable : Exports.Variables)
+				for(const FAnimNextExport& Export : Exports.Exports)
 				{
-					if(Variable.Name == InSoftVariableReference.GetName())
+					if(Export.Identifier == InSoftVariableReference.GetName())
 					{
-						return Variable.Type;
+						return Export.Data.Get<FAnimNextVariableDeclarationData>().Type;
 					}
 				}
 			}
@@ -792,26 +803,12 @@ FAnimNextParamType FUtils::FindVariableType(const FAnimNextSoftVariableReference
 
 bool FUtils::GetExportedVariablesForAsset(const FAssetData& InAsset, FAnimNextAssetRegistryExports& OutExports)
 {
-	const FString TagValue = InAsset.GetTagValueRef<FString>(UE::AnimNext::ExportsAnimNextAssetRegistryTag);
-	return FAnimNextAssetRegistryExports::StaticStruct()->ImportText(*TagValue, &OutExports, nullptr, PPF_None, nullptr, FAnimNextAssetRegistryExports::StaticStruct()->GetName()) != nullptr;
+	return GetExportsOfTypeForAsset<FAnimNextVariableDeclarationData>(InAsset, OutExports);
 }
 
 bool FUtils::GetExportedVariablesFromAssetRegistry(TMap<FAssetData, FAnimNextAssetRegistryExports>& OutExports)
-{
-	TArray<FAssetData> AssetData;
-	IAssetRegistry::GetChecked().GetAssetsByTags({UE::AnimNext::ExportsAnimNextAssetRegistryTag}, AssetData);
-
-	for (const FAssetData& Asset : AssetData)
-	{
-		const FString TagValue = Asset.GetTagValueRef<FString>(UE::AnimNext::ExportsAnimNextAssetRegistryTag);
-		FAnimNextAssetRegistryExports AssetExports;
-		if (FAnimNextAssetRegistryExports::StaticStruct()->ImportText(*TagValue, &AssetExports, nullptr, PPF_None, nullptr, FAnimNextAssetRegistryExports::StaticStruct()->GetName()) != nullptr)
-		{
-			OutExports.Add(Asset, MoveTemp(AssetExports));
-		}
-	}
-
-	return OutExports.Num() > 0;
+{	
+	return GetExportsOfTypeFromAssetRegistry<FAnimNextVariableDeclarationData>(OutExports);
 }
 
 void FUtils::GetAssetFunctions(const UAnimNextRigVMAssetEditorData* InEditorData, FRigVMGraphFunctionHeaderArray& OutExports)
@@ -834,14 +831,6 @@ void FUtils::GetAssetPrivateFunctions(const UAnimNextRigVMAssetEditorData* InEdi
 	}
 }
 
-bool FUtils::GetExportedFunctionsForAsset(const FAssetData& InAsset, FName Tag, FRigVMGraphFunctionHeaderArray& OutExports)
-{
-	const FString TagValue = InAsset.GetTagValueRef<FString>(Tag);
-	const FArrayProperty* HeadersProperty = CastField<FArrayProperty>(FRigVMGraphFunctionHeaderArray::StaticStruct()->FindPropertyByName(TEXT("Headers")));
-	HeadersProperty->ImportText_Direct(*TagValue, &OutExports, nullptr, EPropertyPortFlags::PPF_None);
-	return OutExports.Headers.Num() > 0;
-}
-
 bool FUtils::GetExportedFunctionsFromAssetRegistry(FName Tag, TMap<FAssetData, FRigVMGraphFunctionHeaderArray>& OutExports)
 {
 	TArray<FAssetData> AssetData;
@@ -867,7 +856,34 @@ bool FUtils::GetExportedFunctionsFromAssetRegistry(FName Tag, TMap<FAssetData, F
 	return OutExports.Num() > 0;
 }
 
-void FUtils::GetAssetVariables(const UAnimNextRigVMAssetEditorData* InEditorData, FAnimNextAssetRegistryExports& OutExports)
+static void AddParamToSet(const FAnimNextExport& InNewParam, TSet<FAnimNextExport>& OutExports)
+{
+	if(FAnimNextExport* ExistingEntry = OutExports.Find(InNewParam))
+	{
+		const FAnimNextVariableDeclarationData& NewData = InNewParam.Data.Get<FAnimNextVariableDeclarationData>();
+		FAnimNextVariableDeclarationData& ExistingData = ExistingEntry->Data.GetMutable<FAnimNextVariableDeclarationData>();
+		if(ExistingData.Type != NewData.Type)
+		{
+			UE_LOGFMT(LogAnimation, Warning, "Type mismatch between parameter {ParameterName}. {ParamType1} vs {ParamType1}", InNewParam.Identifier, NewData.Type.ToString(), ExistingData.Type.ToString());
+		}
+		ExistingData.Flags |= NewData.Flags;
+	}
+	else
+	{
+		OutExports.Add(InNewParam);
+	}
+}
+
+void FUtils::GetAssetVariableExports(const UAnimNextRigVMAssetEditorData* EditorData, FAnimNextAssetRegistryExports& OutExports, FAssetRegistryTagsContext Context)
+{
+	OutExports.Exports.Reserve(EditorData->Entries.Num());
+
+	TSet<FAnimNextExport> ExportSet;
+	GetAssetVariableExports(EditorData, ExportSet, Context);
+	OutExports.Exports.Append(ExportSet.Array());
+}
+
+void FUtils::GetAssetVariableExports(const UAnimNextRigVMAssetEditorData* InEditorData, TSet<FAnimNextExport>& OutExports, FAssetRegistryTagsContext Context)
 {
 	for(const UAnimNextRigVMAssetEntry* Entry : InEditorData->Entries)
 	{
@@ -876,8 +892,9 @@ void FUtils::GetAssetVariables(const UAnimNextRigVMAssetEditorData* InEditorData
 			EAnimNextExportedVariableFlags Flags = EAnimNextExportedVariableFlags::Declared;
 			if(ExportInterface->GetExportAccessSpecifier() == EAnimNextExportAccessSpecifier::Public)
 			{
-				Flags |= EAnimNextExportedVariableFlags::Public;
-				OutExports.Variables.Emplace(ExportInterface->GetExportName(), ExportInterface->GetExportType(), Flags);
+				Flags |= EAnimNextExportedVariableFlags::Public;			
+				FAnimNextExport ParameterExport = FAnimNextExport::MakeExport<FAnimNextVariableDeclarationData>(ExportInterface->GetExportName(), ExportInterface->GetExportType(), Flags);
+				AddParamToSet(ParameterExport, OutExports);
 			}
 		}
 		else if(const UAnimNextSharedVariablesEntry* DataInterfaceEntry = Cast<UAnimNextSharedVariablesEntry>(Entry))
@@ -885,68 +902,272 @@ void FUtils::GetAssetVariables(const UAnimNextRigVMAssetEditorData* InEditorData
 			switch(DataInterfaceEntry->Type)
 			{
 			case EAnimNextSharedVariablesType::Asset:
-				OutExports.SharedVariableAssets.Emplace(DataInterfaceEntry->GetAsset());
+				if(const UAnimNextSharedVariables* SharedVariablesAsset = DataInterfaceEntry->Asset)
+				{
+					UAnimNextSharedVariables_EditorData* EditorData = GetEditorData<UAnimNextSharedVariables_EditorData>(SharedVariablesAsset);
+					GetAssetVariableExports(EditorData, OutExports, Context);
+				}
 				break;
 			case EAnimNextSharedVariablesType::Struct:
-				OutExports.SharedVariableStructs.Emplace(DataInterfaceEntry->GetStruct());
+				if(const UScriptStruct* Struct = DataInterfaceEntry->Struct)
+				{
+					GetStructVariableExports(Struct, OutExports);
+				}
 				break;
 			}
 		}
+
+		if(const IAnimNextRigVMGraphInterface* GraphInterface = Cast<IAnimNextRigVMGraphInterface>(Entry))
+		{
+			if (URigVMEdGraph* RigVMEdGraph = GraphInterface->GetEdGraph())
+			{
+				GetSubGraphVariableExportsRecursive(RigVMEdGraph, OutExports);
+			}
+		}
 	}
 }
 
-void FUtils::GetAssetOutlinerItems(const UAnimNextRigVMAssetEditorData* EditorData, FWorkspaceOutlinerItemExports& OutExports, FAssetRegistryTagsContext Context)
+void FUtils::GetStructVariableExports(const UScriptStruct* Struct, TSet<FAnimNextExport>& OutExports)
+{
+	constexpr EAnimNextExportedVariableFlags Flags = EAnimNextExportedVariableFlags::Declared | EAnimNextExportedVariableFlags::Public;
+	for (TFieldIterator<FProperty> It(Struct); It; ++It)
+	{
+		if(const FProperty* Property = *It)
+		{
+			const FName PropertyName = Property->GetFName();
+			FAnimNextExport ParameterExport = FAnimNextExport::MakeExport<FAnimNextVariableDeclarationData>(PropertyName, FAnimNextParamType::FromProperty(Property), Flags);
+			AddParamToSet(ParameterExport, OutExports);
+		}
+	}
+}
+
+void FUtils::GetAssetWorkspaceExports(const UAnimNextRigVMAssetEditorData* EditorData, FWorkspaceOutlinerItemExports& OutExports, FAssetRegistryTagsContext Context)
 {
 	FWorkspaceOutlinerItemExport AssetIdentifier = FWorkspaceOutlinerItemExport(EditorData->GetOuter()->GetFName(), EditorData->GetOuter());
-	for(UAnimNextRigVMAssetEntry* Entry : EditorData->Entries)
+
+	const int32 RootExportIndex = OutExports.Exports.Num() - 1;
+
+	int32 SharedVariablesGroupIndex = INDEX_NONE;
+	auto AddSharedVariablesGroupEntry = [&SharedVariablesGroupIndex, &OutExports, RootExportIndex]()
 	{
-		if(IAnimNextRigVMGraphInterface* GraphInterface = Cast<IAnimNextRigVMGraphInterface>(Entry))
+		const FWorkspaceOutlinerItemExport& ParentExport = OutExports.Exports[RootExportIndex];
+		FWorkspaceOutlinerItemExport& GroupExport = OutExports.Exports.Add_GetRef(FWorkspaceOutlinerItemExport(FName("SharedVariablesGroup"), ParentExport));
+		GroupExport.GetData().InitializeAs<FWorkspaceOutlinerGroupItemData>();
+		GroupExport.GetData().GetMutable<FWorkspaceOutlinerGroupItemData>().GroupName = TEXT("Shared Variables");
+		GroupExport.GetData().GetMutable<FWorkspaceOutlinerGroupItemData>().GroupIcon = *FSlateIconFinder::FindIconBrushForClass(UAnimNextSharedVariables::StaticClass());
+
+		SharedVariablesGroupIndex = OutExports.Exports.Num() - 1;
+	};
+
+	// Referred SharedVariables
+	EditorData->ForEachEntryOfType<UAnimNextSharedVariablesEntry>([&OutExports, &SharedVariablesGroupIndex, AddSharedVariablesGroupEntry](const UAnimNextSharedVariablesEntry* SharedVariablesEntry)
+	{
+		if(SharedVariablesEntry)
 		{
-			if (Entry->IsHiddenInOutliner())
+			if (SharedVariablesGroupIndex == INDEX_NONE)
+			{
+				AddSharedVariablesGroupEntry();
+			}
+
+			switch(SharedVariablesEntry->Type)
 			{
-				if (URigVMEdGraph* RigVMEdGraph = GraphInterface->GetEdGraph())
+			case EAnimNextSharedVariablesType::Asset:
 				{
-					CreateSubGraphsOutlinerItemsRecursive(EditorData, OutExports, AssetIdentifier, INDEX_NONE, RigVMEdGraph, Context);
+					const UAnimNextRigVMAsset* Asset = SharedVariablesEntry->Asset;
+					const FSoftObjectPath ObjectPath(Asset);
+					const FWorkspaceOutlinerItemExport& ParentExport = OutExports.Exports[SharedVariablesGroupIndex];
+					FWorkspaceOutlinerItemExport& ReferenceExport = OutExports.Exports.Add_GetRef(FWorkspaceOutlinerItemExport(FName(ObjectPath.ToString()), ParentExport));
+					ReferenceExport.GetData().InitializeAs<FWorkspaceOutlinerAssetReferenceItemData>();
+					FWorkspaceOutlinerAssetReferenceItemData& Data = ReferenceExport.GetData().GetMutable<FWorkspaceOutlinerAssetReferenceItemData>();
+					Data.ReferredObjectPath = ObjectPath;
+					Data.bShouldExpandReference = false;
 				}
+				break;
+			case EAnimNextSharedVariablesType::Struct:
+				{
+					const UScriptStruct* Struct = SharedVariablesEntry->Struct;
+					const FSoftObjectPath ObjectPath(Struct);
+					const FWorkspaceOutlinerItemExport& ParentExport = OutExports.Exports[SharedVariablesGroupIndex];
+					FWorkspaceOutlinerItemExport& ReferenceExport = OutExports.Exports.Add_GetRef(FWorkspaceOutlinerItemExport(FName(ObjectPath.ToString()), ParentExport));
+					ReferenceExport.GetData().InitializeAs<FWorkspaceOutlinerAssetReferenceItemData>();
+					FWorkspaceOutlinerAssetReferenceItemData& Data = ReferenceExport.GetData().GetMutable<FWorkspaceOutlinerAssetReferenceItemData>();
+					Data.ReferredObjectPath = ObjectPath;
+					Data.bShouldExpandReference = false;
+				}
+				break;
 			}
-			else
+		}
+		return true;
+	});
+
+
+	int32 GraphsGroupIndex = INDEX_NONE;
+	auto AddGraphsGroupEntry = [&GraphsGroupIndex, &OutExports, RootExportIndex]()
+	{
+		const FWorkspaceOutlinerItemExport& ParentExport = OutExports.Exports[RootExportIndex];
+		FWorkspaceOutlinerItemExport& GroupExport = OutExports.Exports.Add_GetRef(FWorkspaceOutlinerItemExport(FName("GraphsGroup"), ParentExport));
+		GroupExport.GetData().InitializeAs<FWorkspaceOutlinerGroupItemData>();
+		GroupExport.GetData().GetMutable<FWorkspaceOutlinerGroupItemData>().GroupName = TEXT("Graphs");
+		GroupExport.GetData().GetMutable<FWorkspaceOutlinerGroupItemData>().GroupIcon = *FAppStyle::GetBrush(TEXT("GraphEditor.EventGraph_24x"));
+
+		GraphsGroupIndex = OutExports.Exports.Num() - 1;
+	};
+
+	EditorData->ForEachEntryOfType<IAnimNextRigVMGraphInterface>([&EditorData, &OutExports, AssetIdentifier, &Context, &GraphsGroupIndex, &AddGraphsGroupEntry](IAnimNextRigVMGraphInterface* GraphInterface)
+	{
+		UAnimNextRigVMAssetEntry* Entry = CastChecked<UAnimNextRigVMAssetEntry>(GraphInterface);		
+		if (Entry->IsHiddenInOutliner())
+		{
+			if (URigVMEdGraph* RigVMEdGraph = GraphInterface->GetEdGraph())
+			{
+				GetSubGraphWorkspaceExportsRecursive(EditorData, OutExports, AssetIdentifier, INDEX_NONE, RigVMEdGraph, Context);
+			}
+		}
+		else
+		{
+			if (GraphsGroupIndex == INDEX_NONE)
 			{
-				FWorkspaceOutlinerItemExport& Export = OutExports.Exports.Add_GetRef(FWorkspaceOutlinerItemExport(Entry->GetEntryName(), AssetIdentifier));
+				AddGraphsGroupEntry();
+			}
+			
+			FWorkspaceOutlinerItemExport& Export = OutExports.Exports.Add_GetRef(FWorkspaceOutlinerItemExport(Entry->GetEntryName(), OutExports.Exports[GraphsGroupIndex]));
 
-				Export.GetData().InitializeAsScriptStruct(FAnimNextGraphOutlinerData::StaticStruct());
-				FAnimNextGraphOutlinerData& GraphData = Export.GetData().GetMutable<FAnimNextGraphOutlinerData>();
-				GraphData.SoftEntryPtr = Entry;
+			Export.GetData().InitializeAsScriptStruct(FAnimNextGraphOutlinerData::StaticStruct());
+			FAnimNextGraphOutlinerData& GraphData = Export.GetData().GetMutable<FAnimNextGraphOutlinerData>();
+			GraphData.SoftEntryPtr = Entry;
 
-				if (URigVMEdGraph* RigVMEdGraph = GraphInterface->GetEdGraph())
-				{
-					int32 ExportIndex = OutExports.Exports.Num() - 1;
-					CreateSubGraphsOutlinerItemsRecursive(EditorData, OutExports, AssetIdentifier, ExportIndex, RigVMEdGraph, Context);
-				}
+			if (URigVMEdGraph* RigVMEdGraph = GraphInterface->GetEdGraph())
+			{					
+				const int32 ExportIndex = OutExports.Exports.Num() - 1;
+				GetSubGraphWorkspaceExportsRecursive(EditorData, OutExports, OutExports.Exports[ExportIndex], ExportIndex, RigVMEdGraph, Context);
 			}
 		}
-	}
+		return true;
+	});	
 
-	CreateFunctionLibraryOutlinerItemsRecursive(EditorData, OutExports, AssetIdentifier, INDEX_NONE, EditorData->GetRigVMGraphFunctionStore()->PublicFunctions, EditorData->GetRigVMGraphFunctionStore()->PrivateFunctions);
+	GetFunctionLibraryWorkspaceExportsRecursive(EditorData, OutExports, AssetIdentifier, INDEX_NONE, EditorData->GetRigVMGraphFunctionStore()->PublicFunctions, EditorData->GetRigVMGraphFunctionStore()->PrivateFunctions);
 }
 
-void ProcessPinAssetReferences(const URigVMPin* InPin, FWorkspaceOutlinerItemExports& OutExports, FWorkspaceOutlinerItemExport& RootExport, int32 ParentExportIndex)
+void FUtils::ProcessPinAssetReferences(const URigVMPin* InPin, FWorkspaceOutlinerItemExports& OutExports, const FWorkspaceOutlinerItemExport& RootExport, int32 ParentExportIndex)
 {
+	auto HandleSoftObjectPath = [&ParentExportIndex, &RootExport, &OutExports](const FSoftObjectPath& InSoftObjectPath) 
+	{
+		// Only add export if object is loaded, or the path actually points to an asset
+		FAssetData ReferenceAssetData;
+		if (InSoftObjectPath.ResolveObject() || IAssetRegistry::GetChecked().TryGetAssetByObjectPath(InSoftObjectPath, ReferenceAssetData) == UE::AssetRegistry::EExists::Exists)
+		{
+			const FWorkspaceOutlinerItemExport& ParentExport = ParentExportIndex == INDEX_NONE ? RootExport : OutExports.Exports[ParentExportIndex];
+			FWorkspaceOutlinerItemExport& ReferenceExport = OutExports.Exports.Add_GetRef(FWorkspaceOutlinerItemExport(FName(InSoftObjectPath.ToString()), ParentExport));
+			ReferenceExport.GetData().InitializeAs<FWorkspaceOutlinerAssetReferenceItemData>();
+			ReferenceExport.GetData().GetMutable<FWorkspaceOutlinerAssetReferenceItemData>().ReferredObjectPath = InSoftObjectPath;
+		}
+	};
+	
 	if (InPin)
 	{
 		const UObject* TypeObject = InPin->GetCPPTypeObject();
 		if (Cast<UClass>(TypeObject))
 		{
-			const FSoftObjectPath ObjectPath(InPin->GetDefaultValue());
+			FSoftObjectPath ObjectPath(InPin->GetDefaultValue());
 			if (ObjectPath.IsValid())
 			{
-				// Only add export if object is loaded, or the path actually points to an asset
-				FAssetData ReferenceAssetData;
-				if (ObjectPath.ResolveObject() || IAssetRegistry::GetChecked().TryGetAssetByObjectPath(ObjectPath, ReferenceAssetData) == UE::AssetRegistry::EExists::Exists)
+				HandleSoftObjectPath(ObjectPath);
+			}
+			else
+			{
+
+				auto HandleDirectPinVariableLink = [&HandleSoftObjectPath](const IAnimNextRigVMVariableInterface* InVariableInterface)
 				{
-					FWorkspaceOutlinerItemExport& ParentExport = ParentExportIndex == INDEX_NONE ? RootExport : OutExports.Exports[ParentExportIndex];
-					FWorkspaceOutlinerItemExport& ReferenceExport = OutExports.Exports.Add_GetRef(FWorkspaceOutlinerItemExport(FName(ObjectPath.ToString()), ParentExport));
-					ReferenceExport.GetData().InitializeAs<FWorkspaceOutlinerAssetReferenceItemData>();
-					ReferenceExport.GetData().GetMutable<FWorkspaceOutlinerAssetReferenceItemData>().ReferredObjectPath = ObjectPath;
+					FString DefaultValue;
+					InVariableInterface->GetDefaultValueString(DefaultValue);
+					const FSoftObjectPath ObjectPath = FSoftObjectPath(DefaultValue);
+
+					if (ObjectPath.IsValid())
+					{
+						HandleSoftObjectPath(ObjectPath);
+					}
+				};
+				
+				// Check for variable nodes linked directly to this pin
+				for (URigVMPin* LinkedPin : InPin->GetLinkedSourcePins())
+				{
+					const URigVMNode* LinkedNode = LinkedPin->GetNode();
+					if (const UAnimNextSharedVariableNode* LinkedSharedVariableNode = Cast<UAnimNextSharedVariableNode>(LinkedNode))
+					{
+						if (LinkedSharedVariableNode->Type == EAnimNextSharedVariablesType::Asset && LinkedSharedVariableNode->Asset)
+						{
+							if (const UAnimNextRigVMAssetEditorData* EditorData = GetEditorData<const UAnimNextRigVMAssetEditorData, const UAnimNextSharedVariables>(LinkedSharedVariableNode->Asset))
+							{
+								if (IAnimNextRigVMVariableInterface* VariableInterface = Cast<IAnimNextRigVMVariableInterface>(EditorData->FindEntry(LinkedSharedVariableNode->GetVariableName())))
+								{
+									HandleDirectPinVariableLink(VariableInterface);
+								}
+							}
+						}
+					}
+					else if (const URigVMVariableNode* LinkedVariableNode = Cast<URigVMVariableNode>(LinkedNode))
+					{
+						if (UAnimNextRigVMAssetEditorData* EditorData = InPin->GetTypedOuter<UAnimNextRigVMAssetEditorData>())
+						{
+							if (IAnimNextRigVMVariableInterface* VariableInterface = Cast<IAnimNextRigVMVariableInterface>(EditorData->FindEntry(LinkedVariableNode->GetVariableName())))
+							{
+								HandleDirectPinVariableLink(VariableInterface);
+							}
+						}
+					}
+				}
+
+				// Check for variable nodes linked to this pin its parent pin
+				if (InPin->GetParentPin())
+				{
+					auto HandleParentPinVariableLink = [&HandleSoftObjectPath](const IAnimNextRigVMVariableInterface* InVariableInterface, const URigVMPin* InPin)
+					{
+						const FAnimNextParamType& ParameterType = InVariableInterface->GetType();
+						if (ParameterType.GetValueType() == FAnimNextParamType::EValueType::Struct)
+						{
+							if (const UScriptStruct* TypeStruct = Cast<UScriptStruct>(ParameterType.GetValueTypeObject()))
+							{
+								if (const FProperty* StructPinProperty = TypeStruct->FindPropertyByName(InPin->GetFName()))
+								{
+									FString DefaultValue;
+									StructPinProperty->ExportTextItem_InContainer(DefaultValue, InVariableInterface->GetValuePtr(), nullptr, nullptr, 0 );
+									const FSoftObjectPath ObjectPath = FSoftObjectPath(DefaultValue);
+									if (ObjectPath.IsValid())
+									{
+										HandleSoftObjectPath(ObjectPath);
+									}
+								}
+							}
+						}
+					};
+					
+					for (URigVMPin* LinkedPin : InPin->GetParentPin()->GetLinkedSourcePins())
+					{
+						const URigVMNode* LinkedNode = LinkedPin->GetNode();
+						if (const UAnimNextSharedVariableNode* LinkedSharedVariableNode = Cast<UAnimNextSharedVariableNode>(LinkedNode))
+						{
+							if (LinkedSharedVariableNode->Type == EAnimNextSharedVariablesType::Asset && LinkedSharedVariableNode->Asset)
+							{
+								if (const UAnimNextRigVMAssetEditorData* EditorData = GetEditorData<const UAnimNextRigVMAssetEditorData, const UAnimNextSharedVariables>(LinkedSharedVariableNode->Asset))
+								{
+									if (IAnimNextRigVMVariableInterface* VariableInterface = Cast<IAnimNextRigVMVariableInterface>(EditorData->FindEntry(LinkedSharedVariableNode->GetVariableName())))
+									{
+										HandleParentPinVariableLink(VariableInterface, InPin);
+									}
+								}
+							}
+						}
+						else if (const URigVMVariableNode* LinkedVariableNode = Cast<URigVMVariableNode>(LinkedNode))
+						{
+							if (UAnimNextRigVMAssetEditorData* EditorData = InPin->GetTypedOuter<UAnimNextRigVMAssetEditorData>())
+							{
+								if (IAnimNextRigVMVariableInterface* VariableInterface = Cast<IAnimNextRigVMVariableInterface>(EditorData->FindEntry(LinkedVariableNode->GetVariableName())))
+								{
+									HandleParentPinVariableLink(VariableInterface, InPin);
+								}
+							}
+						}
+					}
 				}
 			}
 		}
@@ -955,11 +1176,10 @@ void ProcessPinAssetReferences(const URigVMPin* InPin, FWorkspaceOutlinerItemExp
 		{
 			ProcessPinAssetReferences(SubPin, OutExports, RootExport, ParentExportIndex);
 		}
-	}
-	
+	}	
 }
 
-void FUtils::CreateSubGraphsOutlinerItemsRecursive(const UAnimNextRigVMAssetEditorData* EditorData, FWorkspaceOutlinerItemExports& OutExports, FWorkspaceOutlinerItemExport& RootExport, int32 ParentExportIndex, URigVMEdGraph* RigVMEdGraph, FAssetRegistryTagsContext Context)
+void FUtils::GetSubGraphWorkspaceExportsRecursive(const UAnimNextRigVMAssetEditorData* EditorData, FWorkspaceOutlinerItemExports& OutExports, const FWorkspaceOutlinerItemExport& RootExport, int32 ParentExportIndex, URigVMEdGraph* RigVMEdGraph, FAssetRegistryTagsContext Context)
 {
 	if (RigVMEdGraph == nullptr)
 	{
@@ -1013,7 +1233,7 @@ void FUtils::CreateSubGraphsOutlinerItemsRecursive(const UAnimNextRigVMAssetEdit
 			{
 				URigVMCollapseNode* CollapseNode = CastChecked<URigVMCollapseNode>(EditorObject->GetModel()->GetOuter());
 
-				FWorkspaceOutlinerItemExport& ParentExport = ParentExportIndex == INDEX_NONE ? RootExport : OutExports.Exports[ParentExportIndex];
+				const FWorkspaceOutlinerItemExport& ParentExport = ParentExportIndex == INDEX_NONE ? RootExport : OutExports.Exports[ParentExportIndex];
 				FWorkspaceOutlinerItemExport& Export = OutExports.Exports.Add_GetRef(FWorkspaceOutlinerItemExport(CollapseNode->GetFName(), ParentExport));
 				Export.GetData().InitializeAsScriptStruct(FAnimNextCollapseGraphOutlinerData::StaticStruct());
 			
@@ -1021,7 +1241,7 @@ void FUtils::CreateSubGraphsOutlinerItemsRecursive(const UAnimNextRigVMAssetEdit
 				FnGraphData.SoftEditorObject = EditorObject;
 
 				int32 ExportIndex = OutExports.Exports.Num() - 1;
-				CreateSubGraphsOutlinerItemsRecursive(EditorData, OutExports, RootExport, ExportIndex, EditorObject, Context);
+				GetSubGraphWorkspaceExportsRecursive(EditorData, OutExports, RootExport, ExportIndex, EditorObject, Context);
 			}
 		}
 	}
@@ -1041,7 +1261,7 @@ void FUtils::CreateSubGraphsOutlinerItemsRecursive(const UAnimNextRigVMAssetEdit
 				const URigVMLibraryNode* FunctionNode = EditorData->RigVMClient.GetFunctionLibrary()->FindFunction(FunctionIdentifier.GetFunctionFName());
 				if (URigVMGraph* ContainedGraph = FunctionNode->GetContainedGraph())
 				{
-					FWorkspaceOutlinerItemExport& ParentExport = ParentExportIndex == INDEX_NONE ? RootExport : OutExports.Exports[ParentExportIndex];
+					const FWorkspaceOutlinerItemExport& ParentExport = ParentExportIndex == INDEX_NONE ? RootExport : OutExports.Exports[ParentExportIndex];
 					FWorkspaceOutlinerItemExport& Export = OutExports.Exports.Add_GetRef(FWorkspaceOutlinerItemExport(FunctionNode->GetFName(), ParentExport));
 					Export.GetData().InitializeAsScriptStruct(FAnimNextGraphFunctionOutlinerData::StaticStruct());
 					FAnimNextGraphFunctionOutlinerData& FnGraphData = Export.GetData().GetMutable<FAnimNextGraphFunctionOutlinerData>();
@@ -1052,7 +1272,7 @@ void FUtils::CreateSubGraphsOutlinerItemsRecursive(const UAnimNextRigVMAssetEdit
 						FnGraphData.SoftEdGraphNode = EdNode;
 						
 						int32 ExportIndex = OutExports.Exports.Num() - 1;
-						CreateSubGraphsOutlinerItemsRecursive(EditorData, OutExports, RootExport, ExportIndex, ContainedEdGraph, Context);
+						GetSubGraphWorkspaceExportsRecursive(EditorData, OutExports, RootExport, ExportIndex, ContainedEdGraph, Context);
 					}
 				}
 			}
@@ -1060,27 +1280,33 @@ void FUtils::CreateSubGraphsOutlinerItemsRecursive(const UAnimNextRigVMAssetEdit
 	}
 }
 
-void FUtils::CreateFunctionLibraryOutlinerItemsRecursive(const UAnimNextRigVMAssetEditorData* EditorData, FWorkspaceOutlinerItemExports& OutExports, FWorkspaceOutlinerItemExport& RootExport, int32 ParentExportIndex, const TArray<FRigVMGraphFunctionData>& PublicFunctions, const TArray<FRigVMGraphFunctionData>& PrivateFunctions)
+void FUtils::GetFunctionLibraryWorkspaceExportsRecursive(const UAnimNextRigVMAssetEditorData* EditorData, FWorkspaceOutlinerItemExports& OutExports, const FWorkspaceOutlinerItemExport& RootExport, int32 ParentExportIndex, const TArray<FRigVMGraphFunctionData>& PublicFunctions, const TArray<FRigVMGraphFunctionData>& PrivateFunctions)
 {
 	if (PrivateFunctions.Num() > 0 || PublicFunctions.Num() > 0)
 	{
-		CreateFunctionsOutlinerItemsRecursive(EditorData, OutExports, RootExport, ParentExportIndex, PrivateFunctions, false);
-		CreateFunctionsOutlinerItemsRecursive(EditorData, OutExports, RootExport, ParentExportIndex, PublicFunctions, true);
+		const FWorkspaceOutlinerItemExport& ParentExport = ParentExportIndex != INDEX_NONE ? OutExports.Exports[ParentExportIndex] : RootExport;
+		FWorkspaceOutlinerItemExport& GroupExport = OutExports.Exports.Add_GetRef(FWorkspaceOutlinerItemExport(FName("FunctionsGroup"), ParentExport));
+		GroupExport.GetData().InitializeAs<FWorkspaceOutlinerGroupItemData>();
+		GroupExport.GetData().GetMutable<FWorkspaceOutlinerGroupItemData>().GroupName = TEXT("Functions");
+		GroupExport.GetData().GetMutable<FWorkspaceOutlinerGroupItemData>().GroupIcon = *FAppStyle::GetBrush(TEXT("GraphEditor.Function_24x"));
+		const int32 GroupExportIndex = OutExports.Exports.Num() - 1;
+
+		GetFunctionsWorkspaceExportsRecursive(EditorData, OutExports, GroupExport, GroupExportIndex, PrivateFunctions, false);
+		GetFunctionsWorkspaceExportsRecursive(EditorData, OutExports, GroupExport, GroupExportIndex, PublicFunctions, true);
 	}
 }
 
-void FUtils::CreateFunctionsOutlinerItemsRecursive(const UAnimNextRigVMAssetEditorData* EditorData, FWorkspaceOutlinerItemExports& OutExports, FWorkspaceOutlinerItemExport& RootExport, int32 ParentExportIndex, const TArray<FRigVMGraphFunctionData>& Functions, bool bPublicFunctions)
+void FUtils::GetFunctionsWorkspaceExportsRecursive(const UAnimNextRigVMAssetEditorData* EditorData, FWorkspaceOutlinerItemExports& OutExports, const FWorkspaceOutlinerItemExport& RootExport, int32 ParentExportIndex, const TArray<FRigVMGraphFunctionData>& Functions, bool bPublicFunctions)
 {
 	for (const FRigVMGraphFunctionData& FunctionData : Functions)
 	{
-		const URigVMLibraryNode* FunctionNode = EditorData->RigVMClient.GetFunctionLibrary()->FindFunction(FunctionData.Header.LibraryPointer.GetFunctionFName());
-		if (FunctionNode)
+		if (const URigVMLibraryNode* FunctionNode = EditorData->RigVMClient.GetFunctionLibrary()->FindFunction(FunctionData.Header.LibraryPointer.GetFunctionFName()))
 		{
 			if (URigVMGraph* ContainedModelGraph = FunctionNode->GetContainedGraph())
 			{
 				if (URigVMEdGraph* EditorObject = Cast<URigVMEdGraph>(EditorData->GetEditorObjectForRigVMGraph(ContainedModelGraph)))
 				{
-					FWorkspaceOutlinerItemExport& ParentExport = ParentExportIndex == INDEX_NONE ? RootExport : OutExports.Exports[ParentExportIndex];
+					const FWorkspaceOutlinerItemExport& ParentExport = ParentExportIndex == INDEX_NONE ? RootExport : OutExports.Exports[ParentExportIndex];
 					FWorkspaceOutlinerItemExport& Export = OutExports.Exports.Add_GetRef(FWorkspaceOutlinerItemExport(FunctionData.Header.Name, ParentExport));
 
 					Export.GetData().InitializeAsScriptStruct(FAnimNextGraphFunctionOutlinerData::StaticStruct());
@@ -1092,6 +1318,29 @@ void FUtils::CreateFunctionsOutlinerItemsRecursive(const UAnimNextRigVMAssetEdit
 	}
 }
 
+void FUtils::GetSubGraphVariableExportsRecursive(const URigVMEdGraph* RigVMEdGraph, TSet<FAnimNextExport>& OutExports)
+{
+	TArray<URigVMEdGraphNode*> EdNodes;
+	RigVMEdGraph->GetNodesOfClass(EdNodes);
+
+	for (URigVMEdGraphNode* EdNode : EdNodes)
+	{
+		if (URigVMVariableNode* VariableNode = Cast<URigVMVariableNode>(EdNode->GetModelNode()))
+		{
+			OutExports.Add(
+				FAnimNextExport::MakeExport<FAnimNextVariableReferenceData>
+				(
+					VariableNode->GetVariableName(),
+					RigVMEdGraph->GetFName(),
+					VariableNode->GetNodePath(),
+					VariableNode->GetValuePin()->GetPinPath(),
+					VariableNode->IsGetter() ? EAnimNextExportedVariableFlags::Read : EAnimNextExportedVariableFlags::Write
+				)
+			);
+		}
+	}
+}
+
 const FText& FUtils::GetFunctionLibraryDisplayName()
 {
 	static const FText FunctionLibraryName = LOCTEXT("WorkspaceFunctionLibraryName", "Function Library");
@@ -1144,6 +1393,412 @@ FString FUtils::MakeFunctionWrapperEventName(FName InFunctionName)
 	return TEXT("__InternalCall_") + InFunctionName.ToString();
 }
 
+void FUtils::GetVariableNames(UAnimNextRigVMAssetEditorData* InEditorData, TArray<FName>& OutVariableNames, bool bRecursive)
+{
+	for(UAnimNextRigVMAssetEntry* Entry : InEditorData->Entries)
+	{
+		if(UAnimNextVariableEntry* VariableEntry = Cast<UAnimNextVariableEntry>(Entry))
+		{
+			if(VariableEntry->GetExportAccessSpecifier() == EAnimNextExportAccessSpecifier::Public)
+			{
+				OutVariableNames.Add(Entry->GetEntryName());
+			}
+		}
+		else if(UAnimNextSharedVariablesEntry* DataInterfaceEntry = Cast<UAnimNextSharedVariablesEntry>(Entry))
+		{
+			if (bRecursive && DataInterfaceEntry->GetAsset())
+			{
+				UAnimNextSharedVariables_EditorData* EditorData = GetEditorData<UAnimNextSharedVariables_EditorData>(DataInterfaceEntry->GetAsset());
+				GetVariableNames(EditorData, OutVariableNames);
+			}
+		}
+	}
+}
+
+FName FUtils::GetValidVariableName(UAnimNextRigVMAssetEditorData* InEditorData, FName InBaseName)
+{
+	TArray<FName> ExistingNames;
+	GetVariableNames(InEditorData, ExistingNames, true);	
+	return GetValidVariableName(InBaseName, ExistingNames);	
+}
+
+FName FUtils::GetValidVariableName(FName InBaseName, const TArrayView<FName> ExistingNames)
+{
+	auto NameExists = [&ExistingNames](FName InName)
+	{
+		for(FName AdditionalName : ExistingNames)
+		{
+			if(AdditionalName == InName)
+			{
+				return true;
+			}
+		}
+
+		return false;
+	};
+
+	if(!NameExists(InBaseName))
+	{
+		// Early out - name is valid
+		return InBaseName;
+	}
+
+	int32 PostFixIndex = 0;
+	TStringBuilder<128> StringBuilder;
+	while(true)
+	{
+		StringBuilder.Reset();
+		InBaseName.GetDisplayNameEntry()->AppendNameToString(StringBuilder);
+		StringBuilder.Appendf(TEXT("_%d"), PostFixIndex++);
+
+		FName TestName(StringBuilder.ToString()); 
+		if(!NameExists(TestName))
+		{
+			return TestName;
+		}
+	}
+}
+
+void FUtils::DeleteVariable(UAnimNextVariableEntry* VariableEntry, bool bSetupUndoRedo, bool bPrintPythonCommands)
+{
+	if (VariableEntry)
+	{
+		const UAnimNextRigVMAsset* OuterRigVMAsset = VariableEntry->GetTypedOuter<UAnimNextRigVMAsset>();
+		check(OuterRigVMAsset);
+		
+		const FAnimNextSoftVariableReference FindReference(VariableEntry->GetEntryName(), OuterRigVMAsset);
+
+		// Replace references with None
+		UAnimNextRigVMAssetEditorData* EditorData = GetEditorData<UAnimNextRigVMAssetEditorData>(OuterRigVMAsset);
+
+		// Replace reference across project for public variables
+		const bool bIsPublicVariable = VariableEntry->GetExportAccessSpecifier() == EAnimNextExportAccessSpecifier::Public;
+		if (bIsPublicVariable)
+		{
+			ReplaceVariableReferencesAcrossProject(FindReference, FAnimNextSoftVariableReference(), bSetupUndoRedo, bPrintPythonCommands);
+		}
+		else
+		{
+			// Otherwise replace, potential, local references only 
+			ReplaceVariableReferences(EditorData, FindReference, FAnimNextSoftVariableReference(), bSetupUndoRedo, bPrintPythonCommands);		
+		}
+
+		EditorData->RemoveEntry(VariableEntry);
+	}
+}
+
+void FUtils::RenameVariable(UAnimNextVariableEntry* VariableEntry, const FName NewName, bool bSetupUndoRedo, bool bPrintPythonCommands)
+{
+	if (VariableEntry)
+	{
+		UAnimNextRigVMAsset* OuterRigVMAsset = VariableEntry->GetTypedOuter<UAnimNextRigVMAsset>();
+		check(OuterRigVMAsset);
+		
+		const FAnimNextSoftVariableReference FindReference(VariableEntry->GetEntryName(), OuterRigVMAsset);
+		
+		// Rename variable first
+		VariableEntry->SetEntryName(NewName);
+		
+		const FAnimNextSoftVariableReference ReplaceReference(NewName, OuterRigVMAsset);
+		
+		// Replace reference across project for public variables
+		const bool bIsPublicVariable = VariableEntry->GetExportAccessSpecifier() == EAnimNextExportAccessSpecifier::Public;
+		if (bIsPublicVariable)
+		{
+			ReplaceVariableReferencesAcrossProject(FindReference, ReplaceReference, bSetupUndoRedo, bPrintPythonCommands);
+		}
+		else
+		{
+			// Otherwise replace, potential, local references only 
+			ReplaceVariableReferences(GetEditorData(OuterRigVMAsset), FindReference, ReplaceReference, bSetupUndoRedo, bPrintPythonCommands);		
+		}
+	}	
+}
+
+void FUtils::MoveVariableToAsset(UAnimNextVariableEntry* VariableEntry, UAnimNextRigVMAsset* NewOuter, bool bSetupUndoRedo, bool bPrintPythonCommands)
+{
+	if (VariableEntry && NewOuter)
+	{
+		if (UAnimNextRigVMAsset* CurrentOuter = VariableEntry->GetTypedOuter<UAnimNextRigVMAsset>())
+		{
+			UAnimNextRigVMAssetEditorData* CurrentEditorData = GetEditorData(CurrentOuter);
+			UAnimNextRigVMAssetEditorData* NewEditorData = GetEditorData(NewOuter);
+						
+			FString DefaultValueString;
+			VariableEntry->GetDefaultValueString(DefaultValueString);
+			
+			check(NewEditorData->AddVariable(VariableEntry->GetVariableName(), VariableEntry->GetType(), DefaultValueString, bSetupUndoRedo, bPrintPythonCommands));
+
+			const FAnimNextSoftVariableReference FindReference(VariableEntry->GetEntryName(), CurrentOuter);
+			const FAnimNextSoftVariableReference ReplaceReference(VariableEntry->GetEntryName(), NewOuter);
+
+			// Replace reference across project for public variables
+			const bool bIsPublicVariable = VariableEntry->GetExportAccessSpecifier() == EAnimNextExportAccessSpecifier::Public;
+			if (bIsPublicVariable)
+			{
+				ReplaceVariableReferencesAcrossProject(FindReference, ReplaceReference, bSetupUndoRedo, bPrintPythonCommands);
+			}
+			else
+			{
+				// Otherwise replace, potential, local references only 
+				ReplaceVariableReferences(CurrentEditorData, FindReference, ReplaceReference, bSetupUndoRedo, bPrintPythonCommands);		
+			}
+
+			CurrentEditorData->RemoveEntry(VariableEntry, bSetupUndoRedo, bPrintPythonCommands);
+		}
+	}
+}
+
+void FUtils::SetVariableType(UAnimNextVariableEntry* VariableEntry, const FAnimNextParamType NewType, bool bSetupUndoRedo, bool bPrintPythonCommands)
+{
+	if (VariableEntry)
+	{
+		UAnimNextRigVMAsset* OuterRigVMAsset = VariableEntry->GetTypedOuter<UAnimNextRigVMAsset>();
+		check(OuterRigVMAsset);
+		
+		const FAnimNextSoftVariableReference FindReference(VariableEntry->GetEntryName(), OuterRigVMAsset);
+		
+		// Change variable type first
+		VariableEntry->SetType(NewType);
+
+		// Replace reference across project for public variables
+		const bool bIsAPublicVariable = VariableEntry->GetExportAccessSpecifier() == EAnimNextExportAccessSpecifier::Public;
+		if (bIsAPublicVariable)
+		{
+			ReplaceVariableReferencesAcrossProject(FindReference, FindReference, bSetupUndoRedo, bPrintPythonCommands);
+		}
+		else
+		{
+			// Otherwise replace, potential, local references only 
+			ReplaceVariableReferences(GetEditorData(OuterRigVMAsset), FindReference, FindReference, bSetupUndoRedo, bPrintPythonCommands);		
+		}
+	}	
+}
+
+void FUtils::ReplaceVariableReferences(UAnimNextRigVMAssetEditorData* InEditorData, const FAnimNextSoftVariableReference& SoftReferenceToFind, const FAnimNextSoftVariableReference& SoftReferenceToReplaceWith, bool bSetupUndoRedo /*= true*/, bool bPrintPythonCommands /*= true*/)
+{
+	if (InEditorData == nullptr)
+	{
+		return;
+	}
+
+	const FAnimNextVariableReference ReferenceToFind(SoftReferenceToFind);	
+	if (ReferenceToFind.IsNone())
+	{
+		// Cannot replace an empty reference
+		return;
+	}
+		
+	UAnimNextRigVMAsset* ThisOuterAsset = GetAsset(InEditorData);
+	const UAnimNextRigVMAsset* FindOuterAsset = Cast<UAnimNextRigVMAsset>(ReferenceToFind.GetObject());
+	const FAnimNextVariableReference ReferenceToReplaceWith(SoftReferenceToReplaceWith);
+	const UAnimNextRigVMAsset* ReplaceOuterAsset = Cast<UAnimNextRigVMAsset>(ReferenceToReplaceWith.GetObject());
+
+	const bool bReplacingWithNone = ReferenceToReplaceWith.IsNone();
+
+	// Replacing variable ref local to this asset, with another local variable or none
+	const bool bInternalReplace = (ThisOuterAsset == FindOuterAsset) && (FindOuterAsset == ReplaceOuterAsset || bReplacingWithNone);
+	
+	// Replacing with variable from a different object vs existing reference
+	const bool bDifferentVariableSource = FindOuterAsset != ReplaceOuterAsset;
+
+	FAnimNextParamType FindReferenceType = FindVariableType(ReferenceToFind);
+	const FAnimNextParamType ReplaceReferenceType = FindVariableType(ReferenceToReplaceWith);
+
+	// In case ReferenceToFind has already been removed/changed match it up with the replacing type (if not None) 
+	if (!FindReferenceType.IsValid() && !ReferenceToReplaceWith.IsNone())
+	{
+		FindReferenceType = ReplaceReferenceType;
+	}
+	
+	const bool bVariableTypeChanged = FindReferenceType != ReplaceReferenceType;
+	const bool bVariableNameChanged = ReferenceToFind.GetName() != ReferenceToReplaceWith.GetName();
+
+	// Ensure that SharedVariables have been included for replace variable
+	if (const UAnimNextSharedVariables* SharedVariablesReplaceOuter = Cast<UAnimNextSharedVariables>(ReplaceOuterAsset))
+	{
+		bool bHasEntry = false;
+		InEditorData->ForEachEntryOfType<UAnimNextSharedVariablesEntry>([&bHasEntry, SharedVariablesReplaceOuter](const UAnimNextSharedVariablesEntry* Entry)
+		{
+			if (Entry->GetType() == EAnimNextSharedVariablesType::Asset)
+			{
+				if (Entry->GetAsset() == SharedVariablesReplaceOuter)
+				{
+					bHasEntry = true;
+					return false;
+				}
+			}
+
+			return true;
+		});
+
+		if (!bHasEntry)
+		{
+			UE::AnimNext::UncookedOnly::FScopedCompileJob CompilerResults(LOCTEXT("AddSharedAssetJobName", "Add Shared Asset dependency"), { ThisOuterAsset });
+			InEditorData->AddSharedVariables(SharedVariablesReplaceOuter, bSetupUndoRedo, bPrintPythonCommands);
+			InEditorData->RecompileVM();
+		}
+	}			
+	
+	TArray<URigVMEdGraphNode*> EdNodes;
+	for(const UAnimNextRigVMAssetEntry* Entry : InEditorData->Entries)
+	{
+		if(const IAnimNextRigVMGraphInterface* GraphInterface = Cast<IAnimNextRigVMGraphInterface>(Entry))
+		{			
+			if (URigVMEdGraph* RigVMEdGraph = GraphInterface->GetEdGraph())
+			{
+				UAnimNextControllerBase* Controller = CastChecked<UAnimNextControllerBase>(RigVMEdGraph->GetController());
+				// Rely on RigVM controller to handle local variable references
+				if (bInternalReplace)
+				{
+					if (bReplacingWithNone)
+					{
+						Controller->OnExternalVariableRemoved(ReferenceToFind.GetName(), bSetupUndoRedo);
+					}
+					else
+					{
+						if (bVariableTypeChanged)
+						{
+							const FRigVMTemplateArgumentType NewRigVMType = ReplaceReferenceType.ToRigVMTemplateArgument();
+							Controller->OnExternalVariableTypeChanged(ReferenceToFind.GetName(), NewRigVMType.CPPType.ToString(), NewRigVMType.CPPTypeObject.Get(), bSetupUndoRedo);
+						}
+
+						if (bVariableNameChanged)
+						{
+							Controller->OnExternalVariableRenamed(ReferenceToFind.GetName(), ReferenceToReplaceWith.GetName(), bSetupUndoRedo);
+						}
+					}
+				}
+				else
+				{
+					RigVMEdGraph->GetNodesOfClass(EdNodes);
+
+					for (URigVMEdGraphNode* EdNode : EdNodes)
+					{
+						if (URigVMVariableNode* VariableNode = Cast<URigVMVariableNode>(EdNode->GetModelNode()))
+						{
+							const bool bNameMatch = VariableNode->GetVariableName() == ReferenceToFind.GetName();
+
+							if (!bNameMatch)
+							{
+								continue;
+							}								
+							
+							if (UAnimNextSharedVariableNode* SharedVariableNode = Cast<UAnimNextSharedVariableNode>(VariableNode))
+							{
+								const UObject* SearchObject = ReferenceToFind.GetObject(); 
+								const bool bPropertySourceMatch = (SharedVariableNode->Asset && SharedVariableNode->Asset == SearchObject)  || (SharedVariableNode->Struct && SharedVariableNode->Struct == SearchObject);
+								const bool bValidProperty = ReferenceToReplaceWith.ResolveProperty() != nullptr;
+								
+								if (bPropertySourceMatch && (bValidProperty || bReplacingWithNone))
+								{
+									if (bReplacingWithNone)
+									{
+										// [TODO] decide to either remove node or replace with None variable references (keeps node in place but will fail compilation until fixed up)
+										Controller->RemoveNode(SharedVariableNode, bSetupUndoRedo, bPrintPythonCommands);
+										//Controller->RefreshSharedVariableNode(SharedVariableNode->GetFName(), TEXT(""), ReferenceToReplaceWith.GetName(), TEXT(""), nullptr, true, true);	
+									}
+									else
+									{
+										if (bDifferentVariableSource && ReplaceOuterAsset == ThisOuterAsset)
+										{
+											Controller->ReplaceSharedVariableNodeWithVariableNode(SharedVariableNode, ReferenceToReplaceWith.GetName(), bSetupUndoRedo, bPrintPythonCommands);
+										}
+										else
+										{
+											const FRigVMTemplateArgumentType NewRigVMType = ReplaceReferenceType.ToRigVMTemplateArgument();
+											Controller->RefreshSharedVariableNode(SharedVariableNode->GetFName(), ReferenceToReplaceWith.GetObject().GetPath(), ReferenceToReplaceWith.GetName(), NewRigVMType.CPPType.ToString(), NewRigVMType.CPPTypeObject.Get(), bSetupUndoRedo, bPrintPythonCommands);											
+										}
+									}
+								}
+							}
+							else
+							{
+								if (bReplacingWithNone)
+								{
+									// Remove the node for now (can also leave an invalid one causing compile error/warning)
+									Controller->RemoveNode(VariableNode, bSetupUndoRedo, bPrintPythonCommands);
+								}
+								else
+								{
+									if (bDifferentVariableSource)
+									{
+										check(FindOuterAsset == ThisOuterAsset && ReplaceOuterAsset != ThisOuterAsset);
+										{
+											Controller->ReplaceVariableNodeWithSharedVariableNode(VariableNode, ReferenceToReplaceWith.GetName(), ReferenceToReplaceWith.GetObject(), bSetupUndoRedo, bPrintPythonCommands);
+										}
+									}
+									else
+									{					
+										if (bVariableTypeChanged)
+										{
+											const FRigVMTemplateArgumentType NewRigVMType = ReplaceReferenceType.ToRigVMTemplateArgument();
+											Controller->OnExternalVariableTypeChanged(ReferenceToFind.GetName(), NewRigVMType.CPPType.ToString(), NewRigVMType.CPPTypeObject.Get(), bSetupUndoRedo);
+										}
+
+										if (bVariableNameChanged)
+										{
+											Controller->OnExternalVariableRenamed(ReferenceToFind.GetName(), ReferenceToReplaceWith.GetName(), bSetupUndoRedo);
+										}
+									}									
+								}
+							}
+						}
+					}
+				}			
+			}
+		}
+
+		EdNodes.Reset();
+	}	
+}
+
+void FUtils::ReplaceVariableReferencesAcrossProject(const FAnimNextSoftVariableReference& ReferenceToFind, const FAnimNextSoftVariableReference& ReferenceToReplaceWith, bool bSetupUndoRedo, bool bPrintPythonCommands)
+{
+	if (ReferenceToFind.IsNone())
+	{
+		return;
+	}
+	
+	FARFilter AssetFilter;
+	AssetFilter.ClassPaths = { UAnimNextRigVMAsset::StaticClass()->GetClassPathName() };
+	AssetFilter.bRecursiveClasses = true;
+
+	// [TODO] would this miss out assets
+	//AssetFilter.TagsAndValues.Add(UE::AnimNext::ExportsAnimNextAssetRegistryTag, TOptional<FString>());
+		
+	TArray<FAssetData> AssetDatas;
+	FModuleManager::GetModuleChecked<FAssetRegistryModule>("AssetRegistry").Get().GetAssets(AssetFilter, AssetDatas);
+
+	TArray<UAnimNextRigVMAssetEditorData*> EditorDataRequiringReplacing;	
+	for (const FAssetData& AssetData : AssetDatas)
+	{
+		// Relates to [TODO] above
+		ensure(AssetData.FindTag(UE::AnimNext::ExportsAnimNextAssetRegistryTag));
+		
+		FAnimNextAssetRegistryExports AnimNextExports;
+		GetExportsOfTypeForAsset<FAnimNextVariableReferenceData>(AssetData, AnimNextExports);
+
+		for (const FAnimNextExport& Export : AnimNextExports.Exports)
+		{
+			if (ReferenceToFind.GetName() == Export.Identifier)
+			{
+				if (UAnimNextRigVMAsset* RigVMAsset = Cast<UAnimNextRigVMAsset>(AssetData.GetAsset()))
+				{
+					if (UAnimNextRigVMAssetEditorData* EditorData = GetEditorData(RigVMAsset))
+					{
+						EditorDataRequiringReplacing.Add(EditorData);
+					}
+				}
+			}
+		}
+	}
+
+	for (UAnimNextRigVMAssetEditorData* EditorData : EditorDataRequiringReplacing)
+	{
+		ReplaceVariableReferences(EditorData, ReferenceToFind, ReferenceToReplaceWith, bSetupUndoRedo, bPrintPythonCommands);
+	}
+}
 }
 
-#undef LOCTEXT_NAMESPACE
\ No newline at end of file
+#undef LOCTEXT_NAMESPACE
diff --git a/Engine/Plugins/Experimental/AnimNext/Source/AnimNextUncookedOnly/Public/Variables/IAnimNextRigVMVariableInterface.h b/Engine/Plugins/Experimental/AnimNext/Source/AnimNextUncookedOnly/Public/Variables/IAnimNextRigVMVariableInterface.h
--- a/Engine/Plugins/Experimental/AnimNext/Source/AnimNextUncookedOnly/Public/Variables/IAnimNextRigVMVariableInterface.h
+++ b/Engine/Plugins/Experimental/AnimNext/Source/AnimNextUncookedOnly/Public/Variables/IAnimNextRigVMVariableInterface.h
@@ -51,6 +51,7 @@ public:
 	}
 
 	virtual bool GetDefaultValue(const FProperty*& OutProperty, TConstArrayView<uint8>& OutValue) const = 0;
+	virtual bool GetDefaultValueString(FString& OutValueString) const = 0;
 
 	// Set the binding type for this variable (initializes it to default if the struct type is valid)
 	virtual void SetBindingType(UScriptStruct* InBindingTypeStruct, bool bSetupUndoRedo = true) = 0;
@@ -61,6 +62,12 @@ public:
 	// Get the binding for this variable, if any
 	virtual TConstStructView<FAnimNextVariableBindingData> GetBinding() const = 0;
 
+	// Get the category for this variable, empty string if none assigned
+	virtual FStringView GetVariableCategory() const = 0;
+	
+	// Set the category for this variable
+	virtual void SetVariableCategory(FStringView InCategoryName, bool bSetupUndoRedo = true) = 0;
+
 	// Access the memory for the internal value
 	const uint8* GetValuePtr() const
 	{
diff --git a/Engine/Plugins/Experimental/AnimNextAnimGraph/Source/AnimNextAnimGraphEditor/Private/AnimNextAnimGraphEditorModule.cpp b/Engine/Plugins/Experimental/AnimNextAnimGraph/Source/AnimNextAnimGraphEditor/Private/AnimNextAnimGraphEditorModule.cpp
--- a/Engine/Plugins/Experimental/AnimNextAnimGraph/Source/AnimNextAnimGraphEditor/Private/AnimNextAnimGraphEditorModule.cpp
+++ b/Engine/Plugins/Experimental/AnimNextAnimGraph/Source/AnimNextAnimGraphEditor/Private/AnimNextAnimGraphEditorModule.cpp
@@ -91,11 +91,11 @@ class FAnimNextAnimGraphEditorModule : public IModuleInterface
 
 		// --- AnimNextAnimationGraph ---
 		Workspace::FObjectDocumentArgs AnimNextAnimationGraphDocumentArgs(
-			Workspace::FOnRedirectWorkspaceContext::CreateLambda([](UObject* InObject) -> UObject*
+			Workspace::FOnRedirectWorkspaceContext::CreateLambda([](const Workspace::FWorkspaceDocument& Document) -> UObject*
 			{
 				URigVMEdGraph* EdGraph = nullptr;
-
-				UAnimNextAnimationGraph* AnimationGraph = CastChecked<UAnimNextAnimationGraph>(InObject);
+				
+				UAnimNextAnimationGraph* AnimationGraph = CastChecked<UAnimNextAnimationGraph>(Document.GetObject());
 				UAnimNextAnimationGraph_EditorData* EditorData = UncookedOnly::FUtils::GetEditorData<UAnimNextAnimationGraph_EditorData>(AnimationGraph);
 
 				UAnimNextAnimationGraphEntry* AnimationGraphEntry = Cast<UAnimNextAnimationGraphEntry>(EditorData->FindEntry(FRigUnit_AnimNextGraphRoot::DefaultEntryPoint));
@@ -144,8 +144,7 @@ class FAnimNextAnimGraphEditorModule : public IModuleInterface
 
 	virtual void ShutdownModule() override
 	{
-		ISettingsModule* SettingsModule = FModuleManager::GetModulePtr<ISettingsModule>("Settings");
-		if (SettingsModule)
+		if (ISettingsModule* SettingsModule = FModuleManager::GetModulePtr<ISettingsModule>("Settings"))
 		{
 			SettingsModule->UnregisterSettings("Project", "Plugins", "AnimNextAnimGraph");
 		}
@@ -185,14 +184,14 @@ class FAnimNextAnimGraphEditorModule : public IModuleInterface
 		FAnimationGraphMenuExtensions::UnregisterMenus();
 	}
 
-	static void CollectContextMenuActions(const TWeakPtr<UE::Workspace::IWorkspaceEditor>& InWorkspaceEditorWeak, FGraphContextMenuBuilder& InContexMenuBuilder, const FActionMenuContextData& ActionMenuContextData)
+	static void CollectContextMenuActions(const TWeakPtr<UE::Workspace::IWorkspaceEditor>& InWorkspaceEditorWeak, FGraphContextMenuBuilder& InContextMenuBuilder, const FActionMenuContextData& ActionMenuContextData)
 	{
-		if (const URigVMEdGraph* RigVMEdGraph = Cast<URigVMEdGraph>(InContexMenuBuilder.CurrentGraph))
+		if (const URigVMEdGraph* RigVMEdGraph = Cast<URigVMEdGraph>(InContextMenuBuilder.CurrentGraph))
 		{
 			if (RigVMEdGraph->GetModel()->GetSchemaClass() == UAnimNextAnimationGraphSchema::StaticClass())
 			{
-				TArray<FAnimNextAssetRegistryExports> ManifestExports;
-				UE::AnimNext::UncookedOnly::FAnimGraphUtils::GetExportedManifestNodesFromAssetRegistry(ManifestExports);
+				TMap<FAssetData, FAnimNextAssetRegistryExports> PerAssetManifestExports;
+				UE::AnimNext::UncookedOnly::FAnimGraphUtils::GetExportedManifestNodesFromAssetRegistry(PerAssetManifestExports);
 
 				TArray<FAssetData> WorkspaceAssets;
 				if (TSharedPtr<UE::Workspace::IWorkspaceEditor> WorkspaceEditor = InWorkspaceEditorWeak.Pin())
@@ -200,37 +199,22 @@ class FAnimNextAnimGraphEditorModule : public IModuleInterface
 					WorkspaceEditor->GetAssets(WorkspaceAssets);
 				}
 
-				for (const FAnimNextAssetRegistryExports& ManifestExport : ManifestExports)
+				for (const TPair<FAssetData, FAnimNextAssetRegistryExports>& ManifestExports : PerAssetManifestExports)
 				{
-					for (const FAnimNextAssetRegistryManifestNode& ManifestNodeData : ManifestExport.ManifestNodes)
+					for (const FAnimNextExport& Export : ManifestExports.Value.Exports)
 					{
-						bool bIncludeManifestNode = ActionMenuContextData.bShowGlobalManifestNodes;
-
-						if (!ActionMenuContextData.bShowGlobalManifestNodes)
-						{
-							if (WorkspaceAssets.Num() > 0)
-							{
-								for (const FAssetData& WorkspaceAsset : WorkspaceAssets)
-								{
-									if (WorkspaceAsset.PackageName != ManifestNodeData.ModelGraph.GetLongPackageName())
-									{
-										continue;
-									}
-									bIncludeManifestNode = true;
-								}
-							}
-						}
-
+						const FAnimNextManifestNodeData& ManifestNodeData = Export.Data.Get<FAnimNextManifestNodeData>();
+						const bool bIncludeManifestNode = ActionMenuContextData.bShowGlobalManifestNodes || WorkspaceAssets.Contains(ManifestExports.Key);
 						if (bIncludeManifestNode)
 						{
-							InContexMenuBuilder.AddAction(MakeShared<FAnimNextSchemaAction_AddManifestNode>(ManifestNodeData));
+							InContextMenuBuilder.AddAction(MakeShared<FAnimNextSchemaAction_AddManifestNode>(ManifestNodeData));
 						}
 					}
 				}
 			}
 			else if (RigVMEdGraph->GetModel()->GetSchemaClass() == UAnimNextEventGraphSchema::StaticClass())
 			{
-				InContexMenuBuilder.AddAction(MakeShared<FAnimNextSchemaAction_NotifyEvent>());
+				InContextMenuBuilder.AddAction(MakeShared<FAnimNextSchemaAction_NotifyEvent>());
 			}
 		}
 
@@ -242,7 +226,7 @@ class FAnimNextAnimGraphEditorModule : public IModuleInterface
 		{
 			if (ActionMenuContextData.RigVMSchema != nullptr && ActionMenuContextData.RigVMSchema->SupportsUnitFunction(ActionMenuContextData.RigVMController, Function))
 			{
-				UE::AnimNext::Editor::FUtils::AddSchemaRigUnitAction(UAnimNextTraitStackUnitNode::StaticClass(), Struct, *Function, InContexMenuBuilder);
+				UE::AnimNext::Editor::FUtils::AddSchemaRigUnitAction(UAnimNextTraitStackUnitNode::StaticClass(), Struct, *Function, InContextMenuBuilder);
 			}
 		}
 	}
diff --git a/Engine/Plugins/Experimental/AnimNextAnimGraph/Source/AnimNextAnimGraphEditor/Private/Graph/AnimNextAnimationGraphMenuExtensions.cpp b/Engine/Plugins/Experimental/AnimNextAnimGraph/Source/AnimNextAnimGraphEditor/Private/Graph/AnimNextAnimationGraphMenuExtensions.cpp
--- a/Engine/Plugins/Experimental/AnimNextAnimGraph/Source/AnimNextAnimGraphEditor/Private/Graph/AnimNextAnimationGraphMenuExtensions.cpp
+++ b/Engine/Plugins/Experimental/AnimNextAnimGraph/Source/AnimNextAnimGraphEditor/Private/Graph/AnimNextAnimationGraphMenuExtensions.cpp
@@ -13,7 +13,6 @@
 #include "Graph/RigDecorator_AnimNextCppTrait.h"
 #include "Graph/RigUnit_AnimNextRunAnimationGraph_v1.h"
 #include "Graph/RigUnit_AnimNextRunAnimationGraph_v2.h"
-#include "Graph/RigUnit_AnimNextTraitStack.h"
 #include "Graph/RigVMTrait_AnimNextPublicVariables.h"
 #include "RigVMModel/RigVMController.h"
 #include "RigVMModel/Nodes/RigVMUnitNode.h"
@@ -206,9 +205,9 @@ void FAnimationGraphMenuExtensions::RegisterMenus()
 						FAnimNextAssetRegistryExports Exports;
 						if(UE::AnimNext::UncookedOnly::FUtils::GetExportedVariablesForAsset(InAssetData, Exports))
 						{
-							for(const FAnimNextAssetRegistryExportedVariable& Export : Exports.Variables)
+							for(const FAnimNextExport& Export : Exports.Exports)
 							{
-								if((Export.GetFlags() & EAnimNextExportedVariableFlags::Public) != EAnimNextExportedVariableFlags::NoFlags)
+								if((static_cast<EAnimNextExportedVariableFlags>(Export.Data.Get<FAnimNextVariableDeclarationData>().Flags) & EAnimNextExportedVariableFlags::Public) != EAnimNextExportedVariableFlags::NoFlags)
 								{
 									return false;
 								}
@@ -258,7 +257,7 @@ void FAnimationGraphMenuExtensions::UnregisterMenus()
 	}
 }
 
-void FAnimationGraphMenuExtensions::AddManifestSection(UToolMenu* InMenu, const UAnimNextEdGraphNode* AnimNextEdGraphNode, URigVMNode* ModelNode)
+void FAnimationGraphMenuExtensions::AddManifestSection(UToolMenu* InMenu, const UAnimNextEdGraphNode* AnimNextEdGraphNode, const URigVMNode* ModelNode)
 {
 	FToolMenuSection& Section = InMenu->AddSection("AnimNextManifestNodeActions", LOCTEXT("AnimNextManifestNodeActionsMenuHeader", "Manifest"));
 
@@ -300,7 +299,7 @@ void FAnimationGraphMenuExtensions::AddManifestSection(UToolMenu* InMenu, const
 		Section.AddMenuEntry(
 			"RemoveTraitFromManifest",
 			LOCTEXT("RemoveTraitFromManifest", "Remove from Manifest"),
-			LOCTEXT("RemoveTraitFromManifest_Tooltip", "Remvoes this Trait Stack from the Manifest."),
+			LOCTEXT("RemoveTraitFromManifest_Tooltip", "Removes this Trait Stack from the Manifest."),
 			FSlateIcon(),
 			FUIAction(FExecuteAction::CreateLambda([EdNodeWeak]()
 				{
diff --git a/Engine/Plugins/Experimental/AnimNextAnimGraph/Source/AnimNextAnimGraphEditor/Private/Graph/AnimNextAnimationGraphMenuExtensions.h b/Engine/Plugins/Experimental/AnimNextAnimGraph/Source/AnimNextAnimGraphEditor/Private/Graph/AnimNextAnimationGraphMenuExtensions.h
--- a/Engine/Plugins/Experimental/AnimNextAnimGraph/Source/AnimNextAnimGraphEditor/Private/Graph/AnimNextAnimationGraphMenuExtensions.h
+++ b/Engine/Plugins/Experimental/AnimNextAnimGraph/Source/AnimNextAnimGraphEditor/Private/Graph/AnimNextAnimationGraphMenuExtensions.h
@@ -18,7 +18,7 @@ struct FAnimationGraphMenuExtensions
 	static void UnregisterMenus();
 
 private:
-	static void AddManifestSection(UToolMenu* InMenu, const UAnimNextEdGraphNode* AnimNextEdGraphNode, URigVMNode* ModelNode);
+	static void AddManifestSection(UToolMenu* InMenu, const UAnimNextEdGraphNode* AnimNextEdGraphNode, const URigVMNode* ModelNode);
 };
 
 }
diff --git a/Engine/Plugins/Experimental/AnimNextAnimGraph/Source/AnimNextAnimGraphTestSuite/Private/AnimNextAnimGraphTraitGraphTest.cpp b/Engine/Plugins/Experimental/AnimNextAnimGraph/Source/AnimNextAnimGraphTestSuite/Private/AnimNextAnimGraphTraitGraphTest.cpp
--- a/Engine/Plugins/Experimental/AnimNextAnimGraph/Source/AnimNextAnimGraphTestSuite/Private/AnimNextAnimGraphTraitGraphTest.cpp
+++ b/Engine/Plugins/Experimental/AnimNextAnimGraph/Source/AnimNextAnimGraphTestSuite/Private/AnimNextAnimGraphTraitGraphTest.cpp
@@ -421,28 +421,41 @@ bool FAnimationAnimNextEditorTest_GraphManifest::RunTest(const FString& InParame
 		UE_RETURN_ON_ERROR(TraitStackNode->IsExposedToManifest() == true, TEXT("FAnimationAnimNextEditorTest_GraphManifest -> Created node should be at Manifest after Redo."));
 		UE_RETURN_ON_ERROR(UE::AnimNext::UncookedOnly::FAnimGraphUtils::IsExposedToManifest(TraitStackNode) == true, TEXT("FAnimationAnimNextEditorTest_GraphManifest -> FAnimGraphUtils::IsExposedToManifest should return true after Redo."));
 
-		// --- Check Regiostry Exports ---
+		// --- Check Registry Exports ---
 		FAnimNextAssetRegistryExports OutExports;
+
+		auto CountNumberOfManifestNodes = [&OutExports]()-> int32
+		{
+			int32 NumNodes = 0;			
+			OutExports.ForEachExportOfType<FAnimNextManifestNodeData>([&NumNodes](const FAnimNextManifestNodeData& Node) { ++NumNodes; return true; });
+			return NumNodes;
+		};
+		
 		UE::AnimNext::UncookedOnly::FAnimGraphUtils::GetAssetManifestNodesRegistryExports(EditorData, OutExports);
-		UE_RETURN_ON_ERROR(OutExports.ManifestNodes.Num() == 1, TEXT("FAnimationAnimNextEditorTest_GraphManifest -> Asset Registry Exports Num should be 1."));
+		UE_RETURN_ON_ERROR(CountNumberOfManifestNodes() == 1, TEXT("FAnimationAnimNextEditorTest_GraphManifest -> Asset Registry Exports Num should be 1."));
+		OutExports.Exports.Reset();
 
 		// --- Undo Add to Manifest ---
 		AnimNextController->Undo();
 		UE_RETURN_ON_ERROR(TraitStackNode->IsExposedToManifest() == false, TEXT("FAnimationAnimNextEditorTest_GraphManifest -> Created node should not be at Manifest after Undo."));
 		UE_RETURN_ON_ERROR(UE::AnimNext::UncookedOnly::FAnimGraphUtils::IsExposedToManifest(TraitStackNode) == false, TEXT("FAnimationAnimNextEditorTest_GraphManifest -> FAnimGraphUtils::IsExposedToManifest should return false after Undo."));
+		OutExports.Exports.Reset();
 
-		// --- Check Regiostry Exports ---
+		// --- Check Registry Exports ---
 		UE::AnimNext::UncookedOnly::FAnimGraphUtils::GetAssetManifestNodesRegistryExports(EditorData, OutExports);
-		UE_RETURN_ON_ERROR(OutExports.ManifestNodes.Num() == 0, TEXT("FAnimationAnimNextEditorTest_GraphManifest -> Asset Registry Exports Num should be 0 after Undo."));
+		UE_RETURN_ON_ERROR(CountNumberOfManifestNodes() == 0, TEXT("FAnimationAnimNextEditorTest_GraphManifest -> Asset Registry Exports Num should be 0 after Undo."));
+		OutExports.Exports.Reset();
 
 		// --- Redo Add to Manifest ---
 		AnimNextController->Redo();
 		UE_RETURN_ON_ERROR(TraitStackNode->IsExposedToManifest() == true, TEXT("FAnimationAnimNextEditorTest_GraphManifest -> Created node should be at Manifest after Redo."));
 		UE_RETURN_ON_ERROR(UE::AnimNext::UncookedOnly::FAnimGraphUtils::IsExposedToManifest(TraitStackNode) == true, TEXT("FAnimationAnimNextEditorTest_GraphManifest -> FAnimGraphUtils::IsExposedToManifest should return true after Redo."));
+		OutExports.Exports.Reset();
 
-		// --- Check Regiostry Exports ---
+		// --- Check Registry Exports ---
 		UE::AnimNext::UncookedOnly::FAnimGraphUtils::GetAssetManifestNodesRegistryExports(EditorData, OutExports);
-		UE_RETURN_ON_ERROR(OutExports.ManifestNodes.Num() == 1, TEXT("FAnimationAnimNextEditorTest_GraphManifest -> Asset Registry Exports Num should be 1 after Redo."));
+		UE_RETURN_ON_ERROR(CountNumberOfManifestNodes() == 1, TEXT("FAnimationAnimNextEditorTest_GraphManifest -> Asset Registry Exports Num should be 1 after Redo."));
+		OutExports.Exports.Reset();
 	}
 
 	Tests::FUtils::CleanupAfterTests();
diff --git a/Engine/Plugins/Experimental/AnimNextAnimGraph/Source/AnimNextAnimGraphUncookedOnly/Internal/AnimGraphUncookedOnlyUtils.h b/Engine/Plugins/Experimental/AnimNextAnimGraph/Source/AnimNextAnimGraphUncookedOnly/Internal/AnimGraphUncookedOnlyUtils.h
--- a/Engine/Plugins/Experimental/AnimNextAnimGraph/Source/AnimNextAnimGraphUncookedOnly/Internal/AnimGraphUncookedOnlyUtils.h
+++ b/Engine/Plugins/Experimental/AnimNextAnimGraph/Source/AnimNextAnimGraphUncookedOnly/Internal/AnimGraphUncookedOnlyUtils.h
@@ -3,6 +3,7 @@
 #pragma once
 
 #include "CoreMinimal.h"
+#include "AssetRegistry/AssetData.h"
 
 #define UE_API ANIMNEXTANIMGRAPHUNCOOKEDONLY_API
 
@@ -27,7 +28,7 @@ struct FAnimGraphUtils
 	static UE_API void GetAssetManifestNodesRegistryExports(const UAnimNextRigVMAssetEditorData* InEditorData, FAnimNextAssetRegistryExports& OutExports);
 
 	// Gets all the manifest node defs that are exported to the asset registry
-	static UE_API bool GetExportedManifestNodesFromAssetRegistry(TArray<FAnimNextAssetRegistryExports>& OutExports);
+	static UE_API bool GetExportedManifestNodesFromAssetRegistry(TMap<FAssetData, FAnimNextAssetRegistryExports>& OutExports);
 
 	// Returns true if the node is exposed to the manifest
 	static UE_API bool IsExposedToManifest(const URigVMNode* InModelNode);
diff --git a/Engine/Plugins/Experimental/AnimNextAnimGraph/Source/AnimNextAnimGraphUncookedOnly/Internal/AnimNextAnimGraphWorkspaceAssetUserData.h b/Engine/Plugins/Experimental/AnimNextAnimGraph/Source/AnimNextAnimGraphUncookedOnly/Internal/AnimNextAnimGraphWorkspaceAssetUserData.h
--- a/Engine/Plugins/Experimental/AnimNextAnimGraph/Source/AnimNextAnimGraphUncookedOnly/Internal/AnimNextAnimGraphWorkspaceAssetUserData.h
+++ b/Engine/Plugins/Experimental/AnimNextAnimGraph/Source/AnimNextAnimGraphUncookedOnly/Internal/AnimNextAnimGraphWorkspaceAssetUserData.h
@@ -31,4 +31,6 @@ private:
 	GENERATED_BODY()
 
 	virtual void GetAssetRegistryTags(FAssetRegistryTagsContext Context) const override;
+	
+	mutable FWorkspaceOutlinerItemExports CachedExports;
 };
diff --git a/Engine/Plugins/Experimental/AnimNextAnimGraph/Source/AnimNextAnimGraphUncookedOnly/Internal/Graph/AnimGraphEditorSchemaActions.cpp b/Engine/Plugins/Experimental/AnimNextAnimGraph/Source/AnimNextAnimGraphUncookedOnly/Internal/Graph/AnimGraphEditorSchemaActions.cpp
--- a/Engine/Plugins/Experimental/AnimNextAnimGraph/Source/AnimNextAnimGraphUncookedOnly/Internal/Graph/AnimGraphEditorSchemaActions.cpp
+++ b/Engine/Plugins/Experimental/AnimNextAnimGraph/Source/AnimNextAnimGraphUncookedOnly/Internal/Graph/AnimGraphEditorSchemaActions.cpp
@@ -2,8 +2,6 @@
 
 #include "AnimGraphEditorSchemaActions.h"
 #include "EditorUtils.h"
-#include "AnimNextEdGraph.h"
-#include "AnimNextEdGraphNode.h"
 #include "EdGraphNode_Comment.h"
 #include "Editor.h"
 #include "Misc/StringOutputDevice.h"
@@ -15,13 +13,10 @@
 #include "Framework/Application/SlateApplication.h"
 #include "UncookedOnlyUtils.h"
 #include "AnimNextController.h"
-#include "AnimGraphUncookedOnlyUtils.h"
 #include "Editor/RigVMEditorTools.h"
 #include "Exporters/Exporter.h"
 #include "UnrealExporter.h"
-#include "AnimNextUnitNode.h"
 #include "PersonaModule.h"
-#include "ScopedTransaction.h"
 #include "Modules/ModuleManager.h"
 #include "RigVMFunctions/Execution/RigVMFunction_UserDefinedEvent.h"
 #include "Widgets/Input/STextEntryPopup.h"
@@ -35,7 +30,7 @@ static const FText ManifestMenuElementPrefix(LOCTEXT("ManifestMenuElementCategor
 
 // --- Manifest Node ---
 
-FAnimNextSchemaAction_AddManifestNode::FAnimNextSchemaAction_AddManifestNode(const FAnimNextAssetRegistryManifestNode& InManifestNodeData, const FText& InKeywords)
+FAnimNextSchemaAction_AddManifestNode::FAnimNextSchemaAction_AddManifestNode(const FAnimNextManifestNodeData& InManifestNodeData, const FText& InKeywords)
 	: FAnimNextSchemaAction(FText::Format(Private::ManifestMenuElementPrefix, FText::FromString(InManifestNodeData.NodeCategory)), FText::FromString(InManifestNodeData.MenuDesc), FText::FromString(InManifestNodeData.ToolTip), InKeywords)
 	, ModelGraph(InManifestNodeData.ModelGraph)
 	, NodeName(InManifestNodeData.NodeName)
@@ -48,12 +43,6 @@ UEdGraphNode* FAnimNextSchemaAction_AddManifestNode::PerformAction(UEdGraph* Par
 	URigVMEdGraphNode* NewNode = nullptr;
 	URigVMEdGraph* EdGraph = Cast<URigVMEdGraph>(ParentGraph);
 
-	UEdGraphPin* FromPin = nullptr;
-	if (FromPins.Num() > 0)
-	{
-		FromPin = FromPins[0];
-	}
-
 	if (Host != nullptr && EdGraph != nullptr)
 	{
 		if (UAnimNextController* Controller = Cast<UAnimNextController>(EdGraph->GetController()))
@@ -69,7 +58,7 @@ UEdGraphNode* FAnimNextSchemaAction_AddManifestNode::PerformAction(UEdGraph* Par
 			FStringOutputDevice Archive;
 			const FExportObjectInnerContext Context;
 
-			UExporter::ExportToOutputDevice(&Context, SourceNode, NULL, Archive, TEXT("copy"), 0, PPF_ExportsNotFullyQualified | PPF_Copy | PPF_Delimited, false, SourceNode->GetOuter());
+			UExporter::ExportToOutputDevice(&Context, SourceNode, nullptr, Archive, TEXT("copy"), 0, PPF_ExportsNotFullyQualified | PPF_Copy | PPF_Delimited, false, SourceNode->GetOuter());
 
 			const FString NodeData = MoveTemp(Archive);
 
diff --git a/Engine/Plugins/Experimental/AnimNextAnimGraph/Source/AnimNextAnimGraphUncookedOnly/Internal/Graph/AnimGraphEditorSchemaActions.h b/Engine/Plugins/Experimental/AnimNextAnimGraph/Source/AnimNextAnimGraphUncookedOnly/Internal/Graph/AnimGraphEditorSchemaActions.h
--- a/Engine/Plugins/Experimental/AnimNextAnimGraph/Source/AnimNextAnimGraphUncookedOnly/Internal/Graph/AnimGraphEditorSchemaActions.h
+++ b/Engine/Plugins/Experimental/AnimNextAnimGraph/Source/AnimNextAnimGraphUncookedOnly/Internal/Graph/AnimGraphEditorSchemaActions.h
@@ -16,7 +16,7 @@ struct FAnimNextSchemaAction_AddManifestNode : public FAnimNextSchemaAction
 
 	FAnimNextSchemaAction_AddManifestNode() = default;
 
-	UE_API FAnimNextSchemaAction_AddManifestNode(const FAnimNextAssetRegistryManifestNode& InManifestNodeData, const FText& InKeywords = FText::GetEmpty());
+	UE_API FAnimNextSchemaAction_AddManifestNode(const FAnimNextManifestNodeData& InManifestNodeData, const FText& InKeywords = FText::GetEmpty());
 
 	// FEdGraphSchemaAction Interface
 	UE_API virtual UEdGraphNode* PerformAction(UEdGraph* ParentGraph, TArray<UEdGraphPin*>& FromPins, const FVector2f& Location, bool bSelectNewNode = true) override;
diff --git a/Engine/Plugins/Experimental/AnimNextAnimGraph/Source/AnimNextAnimGraphUncookedOnly/Private/AnimGraphUncookedOnlyUtils.cpp b/Engine/Plugins/Experimental/AnimNextAnimGraph/Source/AnimNextAnimGraphUncookedOnly/Private/AnimGraphUncookedOnlyUtils.cpp
--- a/Engine/Plugins/Experimental/AnimNextAnimGraph/Source/AnimNextAnimGraphUncookedOnly/Private/AnimGraphUncookedOnlyUtils.cpp
+++ b/Engine/Plugins/Experimental/AnimNextAnimGraph/Source/AnimNextAnimGraphUncookedOnly/Private/AnimGraphUncookedOnlyUtils.cpp
@@ -5,12 +5,10 @@
 #include "AnimNextEdGraphNode.h"
 #include "Graph/RigUnit_AnimNextGraphRoot.h"
 #include "Graph/RigUnit_AnimNextTraitStack.h"
-#include "AnimNextController.h"
 #include "TraitCore/TraitRegistry.h"
 #include "UncookedOnlyUtils.h"
 #include "AnimNextRigVMAssetEditorData.h"
 #include "Graph/AnimNextAnimationGraph_EditorData.h"
-#include "AssetRegistry/IAssetRegistry.h"
 #include "AnimNextTraitStackUnitNode.h"
 
 namespace UE::AnimNext::UncookedOnly
@@ -61,8 +59,7 @@ void FAnimGraphUtils::GetAssetManifestNodesRegistryExports(const UAnimNextRigVMA
 		TArray<UAnimNextEdGraphNode*> AllNodes;
 		EditorData->GetAllNodesOfClass(AllNodes);
 
-		OutExports.ManifestNodes.Reset();
-		OutExports.ManifestNodes.Reserve(AllNodes.Num());
+		OutExports.Exports.Reserve(OutExports.Exports.Num() + AllNodes.Num());
 
 		for (const UAnimNextEdGraphNode* EdNode : AllNodes)
 		{
@@ -74,37 +71,27 @@ void FAnimGraphUtils::GetAssetManifestNodesRegistryExports(const UAnimNextRigVMA
 					FString PackageName;
 					Package->GetName().Split(FString(TEXT("/")), nullptr, &PackageName, ESearchCase::IgnoreCase, ESearchDir::FromEnd);
 
-					OutExports.ManifestNodes.Add(FAnimNextAssetRegistryManifestNode(
-						ModelNode->GetGraph(),
-						*ModelNode->GetName(),
-						PackageName, // for now we set the package as category
-						EdNode->GetNodeTitle(ENodeTitleType::MenuTitle).ToString(),
-						ModelNode->GetToolTipText().ToString()));
+					OutExports.Exports.Add
+					(
+					FAnimNextExport::MakeExport<FAnimNextManifestNodeData>
+						(
+							ModelNode->GetFName(),
+							ModelNode->GetGraph(),
+							*ModelNode->GetName(),
+							PackageName, // for now we set the package as category
+							EdNode->GetNodeTitle(ENodeTitleType::MenuTitle).ToString(),
+							ModelNode->GetToolTipText().ToString()
+						)
+					);
 				}
 			}
 		}
 	}
 }
 
-bool FAnimGraphUtils::GetExportedManifestNodesFromAssetRegistry(TArray<FAnimNextAssetRegistryExports>& OutExports)
+bool FAnimGraphUtils::GetExportedManifestNodesFromAssetRegistry(TMap<FAssetData, FAnimNextAssetRegistryExports>& OutExports)
 {
-	TArray<FAssetData> AssetData;
-	IAssetRegistry::GetChecked().GetAssetsByTags({ UE::AnimNext::ExportsAnimNextAssetRegistryTag }, AssetData);
-
-	for (const FAssetData& Asset : AssetData)
-	{
-		const FString TagValue = Asset.GetTagValueRef<FString>(UE::AnimNext::ExportsAnimNextAssetRegistryTag);
-		FAnimNextAssetRegistryExports AssetExports;
-		if (FAnimNextAssetRegistryExports::StaticStruct()->ImportText(*TagValue, &AssetExports, nullptr, PPF_None, nullptr, FAnimNextAssetRegistryExports::StaticStruct()->GetName()) != nullptr)
-		{
-			if (AssetExports.ManifestNodes.Num() > 0)
-			{
-				OutExports.Add(MoveTemp(AssetExports));
-			}
-		}
-	}
-
-	return OutExports.Num() > 0;
+	return FUtils::GetExportsOfTypeFromAssetRegistry<FAnimNextManifestNodeData>(OutExports);
 }
 
 bool FAnimGraphUtils::IsExposedToManifest(const URigVMNode* InModelNode)
diff --git a/Engine/Plugins/Experimental/AnimNextAnimGraph/Source/AnimNextAnimGraphUncookedOnly/Private/AnimNextAnimGraphWorkspaceAssetUserData.cpp b/Engine/Plugins/Experimental/AnimNextAnimGraph/Source/AnimNextAnimGraphUncookedOnly/Private/AnimNextAnimGraphWorkspaceAssetUserData.cpp
--- a/Engine/Plugins/Experimental/AnimNextAnimGraph/Source/AnimNextAnimGraphUncookedOnly/Private/AnimNextAnimGraphWorkspaceAssetUserData.cpp
+++ b/Engine/Plugins/Experimental/AnimNextAnimGraph/Source/AnimNextAnimGraphUncookedOnly/Private/AnimNextAnimGraphWorkspaceAssetUserData.cpp
@@ -10,22 +10,25 @@ void UAnimNextAnimGraphWorkspaceAssetUserData::GetAssetRegistryTags(FAssetRegist
 {
 	Super::GetAssetRegistryTags(Context);
 
-	FWorkspaceOutlinerItemExports Exports;
+	if (!Context.IsSaving())
 	{
-		UAnimNextAnimationGraph* Asset = CastChecked<UAnimNextAnimationGraph>(GetOuter());
-		const UAnimNextAnimationGraph_EditorData* EditorData = UE::AnimNext::UncookedOnly::FUtils::GetEditorData<UAnimNextAnimationGraph_EditorData>(Asset);
+		CachedExports.Exports.Reset();
 		{
-			FWorkspaceOutlinerItemExport& RootAssetExport = Exports.Exports.Add_GetRef(FWorkspaceOutlinerItemExport(Asset->GetFName(), Asset));
-			RootAssetExport.GetData().InitializeAsScriptStruct(FAnimNextAnimationGraphOutlinerData::StaticStruct());
-			FAnimNextRigVMAssetOutlinerData& Data = RootAssetExport.GetData().GetMutable<FAnimNextRigVMAssetOutlinerData>();
-			Data.SoftAssetPtr = Asset;
-		}
+			UAnimNextAnimationGraph* Asset = CastChecked<UAnimNextAnimationGraph>(GetOuter());
+			const UAnimNextAnimationGraph_EditorData* EditorData = UE::AnimNext::UncookedOnly::FUtils::GetEditorData<UAnimNextAnimationGraph_EditorData>(Asset);
+			{
+				FWorkspaceOutlinerItemExport& RootAssetExport = CachedExports.Exports.Add_GetRef(FWorkspaceOutlinerItemExport(Asset->GetFName(), Asset));
+				RootAssetExport.GetData().InitializeAsScriptStruct(FAnimNextAnimationGraphOutlinerData::StaticStruct());
+				FAnimNextRigVMAssetOutlinerData& Data = RootAssetExport.GetData().GetMutable<FAnimNextRigVMAssetOutlinerData>();
+				Data.SoftAssetPtr = Asset;
+			}
 
-		UE::AnimNext::UncookedOnly::FUtils::GetAssetOutlinerItems(EditorData, Exports, Context);
+			UE::AnimNext::UncookedOnly::FUtils::GetAssetWorkspaceExports(EditorData, CachedExports, Context);
+		}
 	}
 
 	FString TagValue;
-	FWorkspaceOutlinerItemExports::StaticStruct()->ExportText(TagValue, &Exports, nullptr, nullptr, PPF_None, nullptr);
+	FWorkspaceOutlinerItemExports::StaticStruct()->ExportText(TagValue, &CachedExports, nullptr, nullptr, PPF_None, nullptr);
 	Context.AddTag(FAssetRegistryTag(UE::Workspace::ExportsWorkspaceItemsRegistryTag, TagValue, FAssetRegistryTag::TT_Hidden));
 }
 
diff --git a/Engine/Plugins/Experimental/AnimNextAnimGraph/Source/AnimNextAnimGraphUncookedOnly/Private/Graph/AnimNextAnimationGraph_EditorData.cpp b/Engine/Plugins/Experimental/AnimNextAnimGraph/Source/AnimNextAnimGraphUncookedOnly/Private/Graph/AnimNextAnimationGraph_EditorData.cpp
--- a/Engine/Plugins/Experimental/AnimNextAnimGraph/Source/AnimNextAnimGraphUncookedOnly/Private/Graph/AnimNextAnimationGraph_EditorData.cpp
+++ b/Engine/Plugins/Experimental/AnimNextAnimGraph/Source/AnimNextAnimGraphUncookedOnly/Private/Graph/AnimNextAnimationGraph_EditorData.cpp
@@ -7,20 +7,15 @@
 #include "RigVMPythonUtils.h"
 #include "Compilation/AnimNextGetFunctionHeaderCompileContext.h"
 #include "Compilation/AnimNextGetVariableCompileContext.h"
-#include "Compilation/AnimNextGetGraphCompileContext.h"
 #include "Compilation/AnimNextProcessGraphCompileContext.h"
 #include "UncookedOnlyUtils.h"
-#include "AssetRegistry/IAssetRegistry.h"
 #include "Curves/CurveFloat.h"
-#include "Module/AnimNextModule.h"
 #include "Entries/AnimNextAnimationGraphEntry.h"
 #include "Entries/AnimNextSharedVariablesEntry.h"
 #include "Graph/AnimNextAnimationGraph.h"
-#include "Entries/AnimNextEventGraphEntry.h"
 #include "Entries/AnimNextVariableEntry.h"
 #include "Graph/AnimNextAnimationGraphSchema.h"
 #include "Graph/RigDecorator_AnimNextCppTrait.h"
-#include "Graph/RigUnit_AnimNextBeginExecution.h"
 #include "Graph/RigUnit_AnimNextShimRoot.h"
 #include "Graph/RigUnit_AnimNextTraitStack.h"
 #include "Logging/LogScopedVerbosityOverride.h"
@@ -35,7 +30,6 @@
 #include "Entries/AnimNextRigVMAssetEntry.h"
 #include "Graph/RigUnit_AnimNextGraphEvaluator.h"
 #include "Graph/RigUnit_AnimNextGraphRoot.h"
-#include "Variables/AnimNextProgrammaticVariable.h"
 
 #define LOCTEXT_NAMESPACE "AnimNextAnimationGraph_EditorData"
 
@@ -598,7 +592,6 @@ void UAnimNextAnimationGraph_EditorData::OnPreCompileProcessGraphs(const FRigVMC
 				FScopedCompileJob::GetLog().NotifyIntermediateObjectCreation(CreatedObjects.FindChecked(SourceNode), SourceNode);
 			}
 
-			UAnimNextController* TempController = CastChecked<UAnimNextController>(VMClient->GetOrCreateController(TempGraph));
 			TempGraph->SetFlags(RF_Transient);
 			AnimGraphs.Add(TempGraph);
 		}
@@ -746,7 +739,7 @@ bool UAnimNextAnimationGraph_EditorData::CanAddNewEntry(TSubclassOf<UAnimNextRig
 	// Prevent users adding more than one animation graph
 	if(InClass == UAnimNextAnimationGraphEntry::StaticClass())
 	{
-		auto IsAnimNextGraphEntry = [](UAnimNextRigVMAssetEntry* InEntry)
+		auto IsAnimNextGraphEntry = [](const UAnimNextRigVMAssetEntry* InEntry)
 		{
 			if (InEntry)
 			{
@@ -875,6 +868,7 @@ void UAnimNextAnimationGraph_EditorData::InitializeAssetUserData()
 void UAnimNextAnimationGraph_EditorData::GetAnimNextAssetRegistryTags(FAssetRegistryTagsContext& Context, FAnimNextAssetRegistryExports& OutExports) const
 {
 	UE::AnimNext::UncookedOnly::FAnimGraphUtils::GetAssetManifestNodesRegistryExports(this, OutExports);
+	UAnimNextRigVMAssetEditorData::GetAnimNextAssetRegistryTags(Context, OutExports);
 }
 
 #undef LOCTEXT_NAMESPACE
\ No newline at end of file
diff --git a/Engine/Plugins/Experimental/AnimNextStateTree/Source/AnimNextStateTreeEditor/Private/AnimNextStateTreeEditorHost.cpp b/Engine/Plugins/Experimental/AnimNextStateTree/Source/AnimNextStateTreeEditor/Private/AnimNextStateTreeEditorHost.cpp
--- a/Engine/Plugins/Experimental/AnimNextStateTree/Source/AnimNextStateTreeEditor/Private/AnimNextStateTreeEditorHost.cpp
+++ b/Engine/Plugins/Experimental/AnimNextStateTree/Source/AnimNextStateTreeEditor/Private/AnimNextStateTreeEditorHost.cpp
@@ -8,6 +8,7 @@
 #include "StateTreeEditorWorkspaceTabHost.h"
 
 #include "IWorkspaceEditor.h"
+#include "IWorkspaceEditorModule.h"
 
 void FAnimNextStateTreeEditorHost::Init(const TWeakPtr<UE::Workspace::IWorkspaceEditor>& InWeakWorkspaceEditor)
 {
@@ -16,14 +17,15 @@ void FAnimNextStateTreeEditorHost::Init(const TWeakPtr<UE::Workspace::IWorkspace
 	
 	TSharedPtr<UE::Workspace::IWorkspaceEditor> SharedEditor = InWeakWorkspaceEditor.Pin();
 	check(SharedEditor.IsValid());
-	SharedEditor->OnFocussedDocumentChanged().AddSP(this, &FAnimNextStateTreeEditorHost::OnWorkspaceFocussedDocumentChanged);
+	SharedEditor->OnFocusedDocumentChanged().AddSP(this, &FAnimNextStateTreeEditorHost::OnWorkspaceFocusedDocumentChanged);
 }
 
 UStateTree* FAnimNextStateTreeEditorHost::GetStateTree() const
 {
 	if (TSharedPtr<UE::Workspace::IWorkspaceEditor> SharedWorkspaceEditor = WeakWorkspaceEditor.Pin())
 	{
-		const TObjectPtr<UAnimNextStateTree> AnimNextStateTreePtr = SharedWorkspaceEditor->GetFocussedDocument<UAnimNextStateTree>();
+		const UE::Workspace::FWorkspaceDocument& Document = SharedWorkspaceEditor->GetFocusedDocumentOfClass(UAnimNextStateTree::StaticClass());
+		const TObjectPtr<UAnimNextStateTree> AnimNextStateTreePtr = Document.GetTypedObject<UAnimNextStateTree>();
 		return AnimNextStateTreePtr ? AnimNextStateTreePtr->StateTree : nullptr;
 	}	
 
@@ -50,9 +52,9 @@ bool FAnimNextStateTreeEditorHost::CanToolkitSpawnWorkspaceTab() const
 	return true;
 }
 
-void FAnimNextStateTreeEditorHost::OnWorkspaceFocussedDocumentChanged(TObjectPtr<UObject> InObject) const
+void FAnimNextStateTreeEditorHost::OnWorkspaceFocusedDocumentChanged(const UE::Workspace::FWorkspaceDocument& InDocument) const
 {
-	TObjectPtr<UAnimNextStateTree> InStateTree = Cast<UAnimNextStateTree>(InObject);
+	TObjectPtr<UAnimNextStateTree> InStateTree = InDocument.GetTypedObject<UAnimNextStateTree>();
 
 	if (TStrongObjectPtr<UAnimNextStateTree> LastStateTreePinned = WeakLastStateTree.Pin())
 	{
diff --git a/Engine/Plugins/Experimental/AnimNextStateTree/Source/AnimNextStateTreeEditor/Private/AnimNextStateTreeEditorHost.h b/Engine/Plugins/Experimental/AnimNextStateTree/Source/AnimNextStateTreeEditor/Private/AnimNextStateTreeEditorHost.h
--- a/Engine/Plugins/Experimental/AnimNextStateTree/Source/AnimNextStateTreeEditor/Private/AnimNextStateTreeEditorHost.h
+++ b/Engine/Plugins/Experimental/AnimNextStateTree/Source/AnimNextStateTreeEditor/Private/AnimNextStateTreeEditorHost.h
@@ -7,6 +7,7 @@
 namespace UE::Workspace
 {
 	class IWorkspaceEditor;
+	struct FWorkspaceDocument;
 }
 
 class UAnimNextStateTree;
@@ -28,7 +29,7 @@ public:
 	virtual TSharedPtr<UE::StateTreeEditor::FWorkspaceTabHost> GetTabHost() const override;
 
 protected:
-	void OnWorkspaceFocussedDocumentChanged(TObjectPtr<UObject> InObject) const;
+	void OnWorkspaceFocusedDocumentChanged(const UE::Workspace::FWorkspaceDocument& InDocument) const;
 
 	TWeakPtr<UE::Workspace::IWorkspaceEditor> WeakWorkspaceEditor;	
 	FSimpleMulticastDelegate OnStateTreeChangedDelegate;
diff --git a/Engine/Plugins/Experimental/Workspace/Source/WorkspaceEditor/Private/AssetDocumentSummoner.cpp b/Engine/Plugins/Experimental/Workspace/Source/WorkspaceEditor/Private/AssetDocumentSummoner.cpp
--- a/Engine/Plugins/Experimental/Workspace/Source/WorkspaceEditor/Private/AssetDocumentSummoner.cpp
+++ b/Engine/Plugins/Experimental/Workspace/Source/WorkspaceEditor/Private/AssetDocumentSummoner.cpp
@@ -8,7 +8,6 @@
 #include "StructUtils/InstancedStruct.h"
 #include "WorkspaceEditorModule.h"
 #include "Widgets/Docking/SDockTab.h"
-#include "Widgets/Layout/SSpacer.h"
 #include "WorkspaceDocumentState.h"
 #include "SWorkspaceTabWrapper.h"
 #include "WorkspaceTabPayload.h"
@@ -36,9 +35,9 @@ void FAssetDocumentSummoner::SetAllowedClassPaths(TConstArrayView<FTopLevelAsset
 void FAssetDocumentSummoner::OnTabActivated(TSharedPtr<SDockTab> Tab) const
 {
 	TSharedRef<SWorkspaceTabWrapper> TabWrapper = StaticCastSharedRef<SWorkspaceTabWrapper>(Tab->GetContent());
-	if (UObject* DocumentAsset = TabWrapper->GetDocumentObject().Get())
+	if (TabWrapper->GetDocumentObject().Get())
 	{
-		HostingAppPtr.Pin()->SetFocussedDocument(DocumentAsset);
+		HostingAppPtr.Pin()->SetFocusedDocument(TabWrapper->GetDocument());
 	}
 }
 
@@ -54,15 +53,15 @@ void FAssetDocumentSummoner::SaveState(TSharedPtr<SDockTab> Tab, TSharedPtr<FTab
 {
 	if(TSharedPtr<FWorkspaceEditor> WorkspaceEditor = HostingAppPtr.Pin())
 	{
-		UObject* Object = Payload->IsValid() ? FTabPayload_WorkspaceDocument::CastChecked<UObject>(Payload.ToSharedRef()) : nullptr;
-		if(Object)
+		if(UObject* Object = Payload->IsValid() ? FTabPayload_WorkspaceDocument::CastChecked<UObject>(Payload.ToSharedRef()) : nullptr)
 		{
 			TSharedRef<SWorkspaceTabWrapper> TabWrapper = StaticCastSharedRef<SWorkspaceTabWrapper>(Tab->GetContent());
 			FWorkspaceEditorModule& WorkspaceEditorModule = FModuleManager::LoadModuleChecked<FWorkspaceEditorModule>("WorkspaceEditor");
-			const FObjectDocumentArgs* DocumentArgs = WorkspaceEditorModule.FindObjectDocumentType(Object);
+			const FWorkspaceOutlinerItemExport& Export = FTabPayload_WorkspaceDocument::GetExport(Payload.ToSharedRef());
+			const FObjectDocumentArgs* DocumentArgs = WorkspaceEditorModule.FindObjectDocumentType({Export, Object});
 			if(DocumentArgs && DocumentArgs->OnGetDocumentState.IsBound())
 			{
-				WorkspaceEditor->RecordDocumentState(DocumentArgs->OnGetDocumentState.Execute(FWorkspaceEditorContext(WorkspaceEditor.ToSharedRef(), { FTabPayload_WorkspaceDocument::GetExport(Payload.ToSharedRef()), Object }), TabWrapper->GetContent()));
+				WorkspaceEditor->RecordDocumentState(DocumentArgs->OnGetDocumentState.Execute(FWorkspaceEditorContext(WorkspaceEditor.ToSharedRef(), {Export, Object}), TabWrapper->GetContent()));
 			}
 			else
 			{
@@ -82,11 +81,12 @@ TAttribute<FText> FAssetDocumentSummoner::ConstructTabName(const FWorkflowTabSpa
 	}
 
 	FWorkspaceEditorModule& WorkspaceEditorModule = FModuleManager::LoadModuleChecked<FWorkspaceEditorModule>("WorkspaceEditor");
-	if(const FObjectDocumentArgs* DocumentArgs = WorkspaceEditorModule.FindObjectDocumentType(DocumentID))
+	const FWorkspaceOutlinerItemExport& Export = FTabPayload_WorkspaceDocument::GetExport(Info.Payload.ToSharedRef());
+	if(const FObjectDocumentArgs* DocumentArgs = WorkspaceEditorModule.FindObjectDocumentType({Export, DocumentID }))
 	{
 		if(DocumentArgs->OnGetTabName.IsBound())
 		{
-			return DocumentArgs->OnGetTabName.Execute(FWorkspaceEditorContext(WorkspaceEditor.ToSharedRef(), DocumentID, FTabPayload_WorkspaceDocument::GetExport(Info.Payload.ToSharedRef())));
+			return DocumentArgs->OnGetTabName.Execute(FWorkspaceEditorContext(WorkspaceEditor.ToSharedRef(), DocumentID, Export));
 		}
 	}
 
@@ -103,11 +103,10 @@ TAttribute<FText> FAssetDocumentSummoner::ConstructTabName(const FWorkflowTabSpa
 
 bool FAssetDocumentSummoner::IsPayloadSupported(TSharedRef<FTabPayload> Payload) const
 {
-	UObject* Object = FTabPayload_WorkspaceDocument::CastChecked<UObject>(Payload);
-	if(Object)
+	if(UObject* Object = FTabPayload_WorkspaceDocument::CastChecked<UObject>(Payload))
 	{
 		FWorkspaceEditorModule& WorkspaceEditorModule = FModuleManager::LoadModuleChecked<FWorkspaceEditorModule>("WorkspaceEditor");
-		if(const FObjectDocumentArgs* DocumentArgs = WorkspaceEditorModule.FindObjectDocumentType(Object))
+		if(const FObjectDocumentArgs* DocumentArgs = WorkspaceEditorModule.FindObjectDocumentType({FTabPayload_WorkspaceDocument::GetExport(Payload), Object }))
 		{
 			return AllowedClassPaths.Contains(Object->GetClass()->GetClassPathName()) || bAllowUnsupportedClasses;
 		}
@@ -118,8 +117,7 @@ bool FAssetDocumentSummoner::IsPayloadSupported(TSharedRef<FTabPayload> Payload)
 
 TAttribute<FText> FAssetDocumentSummoner::ConstructTabLabelSuffix(const FWorkflowTabSpawnInfo& Info) const
 {
-	UObject* Object = FTabPayload_WorkspaceDocument::CastChecked<UObject>(Info.Payload.ToSharedRef());
-	if(Object)
+	if(UObject* Object = FTabPayload_WorkspaceDocument::CastChecked<UObject>(Info.Payload.ToSharedRef()))
 	{
 		return MakeAttributeLambda([WeakObject = TWeakObjectPtr<UObject>(Object)]()
 		{
@@ -148,8 +146,9 @@ TSharedRef<SWidget> FAssetDocumentSummoner::CreateTabBody(const FWorkflowTabSpaw
 	UObject* DocumentID = FTabPayload_WorkspaceDocument::CastChecked<UObject>(Info.Payload.ToSharedRef());
 	const FWorkspaceEditorModule& WorkspaceEditorModule = FModuleManager::LoadModuleChecked<FWorkspaceEditorModule>("WorkspaceEditor");
 
-	FWorkspaceEditorContext Context(WorkspaceEditor.ToSharedRef(), DocumentID, FTabPayload_WorkspaceDocument::GetExport(Info.Payload.ToSharedRef()));
-	if(const FObjectDocumentArgs* DocumentArgs = WorkspaceEditorModule.FindObjectDocumentType(DocumentID))
+	const FWorkspaceOutlinerItemExport& Export = FTabPayload_WorkspaceDocument::GetExport(Info.Payload.ToSharedRef());
+	FWorkspaceEditorContext Context(WorkspaceEditor.ToSharedRef(), DocumentID, Export);
+	if(const FObjectDocumentArgs* DocumentArgs = WorkspaceEditorModule.FindObjectDocumentType({Export, DocumentID}))
 	{
 		if(DocumentArgs->OnMakeDocumentWidget.IsBound())
 		{
@@ -169,11 +168,12 @@ const FSlateBrush* FAssetDocumentSummoner::GetTabIcon(const FWorkflowTabSpawnInf
 	{
 		UObject* DocumentID = FTabPayload_WorkspaceDocument::CastChecked<UObject>(Info.Payload.ToSharedRef());
 		FWorkspaceEditorModule& WorkspaceEditorModule = FModuleManager::LoadModuleChecked<FWorkspaceEditorModule>("WorkspaceEditor");
-		if(const FObjectDocumentArgs* DocumentArgs = WorkspaceEditorModule.FindObjectDocumentType(DocumentID))
+		const FWorkspaceOutlinerItemExport& Export = FTabPayload_WorkspaceDocument::GetExport(Info.Payload.ToSharedRef());
+		if(const FObjectDocumentArgs* DocumentArgs = WorkspaceEditorModule.FindObjectDocumentType({Export, DocumentID}))
 		{
 			if(DocumentArgs->OnGetTabIcon.IsBound())
 			{
-				return DocumentArgs->OnGetTabIcon.Execute(FWorkspaceEditorContext(WorkspaceEditor.ToSharedRef(), DocumentID, FTabPayload_WorkspaceDocument::GetExport(Info.Payload.ToSharedRef())));
+				return DocumentArgs->OnGetTabIcon.Execute(FWorkspaceEditorContext(WorkspaceEditor.ToSharedRef(), DocumentID, Export));
 			}
 		}
 
diff --git a/Engine/Plugins/Experimental/Workspace/Source/WorkspaceEditor/Private/Outliner/WorkspaceOutlinerHierarchy.cpp b/Engine/Plugins/Experimental/Workspace/Source/WorkspaceEditor/Private/Outliner/WorkspaceOutlinerHierarchy.cpp
--- a/Engine/Plugins/Experimental/Workspace/Source/WorkspaceEditor/Private/Outliner/WorkspaceOutlinerHierarchy.cpp
+++ b/Engine/Plugins/Experimental/Workspace/Source/WorkspaceEditor/Private/Outliner/WorkspaceOutlinerHierarchy.cpp
@@ -17,6 +17,8 @@ namespace UE::Workspace
 
 	void FWorkspaceOutlinerHierarchy::CreateItemsFromAssetData(const FAssetData& AssetData, TArray<FSceneOutlinerTreeItemPtr>& OutItems, TArray<FAssetData>& OutAssets, const FWorkspaceOutlinerItemExport* InParentExport) const
 	{
+		ProcessedAssetData.AddUnique(AssetData);
+		
 		FString TagValue;
 		if(AssetData.GetTagValue(UE::Workspace::ExportsWorkspaceItemsRegistryTag, TagValue))
 		{
@@ -143,6 +145,7 @@ namespace UE::Workspace
 		{
 			TArray<FAssetData> AssetDataEntries; 
 			Workspace->GetAssetDataEntries(AssetDataEntries);
+			ProcessedAssetData.Reset();
 
 			for (const FAssetData& AssetData : AssetDataEntries)
 			{
diff --git a/Engine/Plugins/Experimental/Workspace/Source/WorkspaceEditor/Private/Outliner/WorkspaceOutlinerHierarchy.h b/Engine/Plugins/Experimental/Workspace/Source/WorkspaceEditor/Private/Outliner/WorkspaceOutlinerHierarchy.h
--- a/Engine/Plugins/Experimental/Workspace/Source/WorkspaceEditor/Private/Outliner/WorkspaceOutlinerHierarchy.h
+++ b/Engine/Plugins/Experimental/Workspace/Source/WorkspaceEditor/Private/Outliner/WorkspaceOutlinerHierarchy.h
@@ -25,6 +25,9 @@ namespace UE::Workspace
 	protected:
 		void CreateItemsFromAssetData(const FAssetData& AssetData, TArray<FSceneOutlinerTreeItemPtr>& OutItems, TArray<FAssetData>& OutAssets, const FWorkspaceOutlinerItemExport* InParentExport = nullptr) const;
 	private:
+		friend class SWorkspaceView;
 		TWeakObjectPtr<UWorkspace> WeakWorkspace;
+		// Set of assets referenced by currently from current workspace,(includes UWorkspace::AssetEntries + any expanded asset references)
+		mutable TArray<FAssetData> ProcessedAssetData;
 	};
 }
\ No newline at end of file
diff --git a/Engine/Plugins/Experimental/Workspace/Source/WorkspaceEditor/Private/Outliner/WorkspaceOutlinerMode.cpp b/Engine/Plugins/Experimental/Workspace/Source/WorkspaceEditor/Private/Outliner/WorkspaceOutlinerMode.cpp
--- a/Engine/Plugins/Experimental/Workspace/Source/WorkspaceEditor/Private/Outliner/WorkspaceOutlinerMode.cpp
+++ b/Engine/Plugins/Experimental/Workspace/Source/WorkspaceEditor/Private/Outliner/WorkspaceOutlinerMode.cpp
@@ -193,17 +193,16 @@ TSharedPtr<SWidget> FWorkspaceOutlinerMode::CreateContextMenu()
 	{
 		UWorkspaceItemMenuContext* MenuContext = NewObject<UWorkspaceItemMenuContext>();
 		MenuContext->WeakCommandList = CommandList;
-		
-		Algo::TransformIf(SceneOutliner->GetSelectedItems(), MenuContext->SelectedExports,
-			[](const FSceneOutlinerTreeItemPtr& InItem)
-			{
-				return InItem->IsA<FWorkspaceOutlinerTreeItem>();
-			},
-			[](const FSceneOutlinerTreeItemPtr& InItem)
+
+		for (const FSceneOutlinerTreeItemPtr& Item : SceneOutliner->GetSelectedItems())
+		{
+			if (Item->IsA<FWorkspaceOutlinerTreeItem>())
 			{
-				return InItem->CastTo<FWorkspaceOutlinerTreeItem>()->Export;
-			});	
-		
+				MenuContext->SelectedExports.Add(Item->CastTo<FWorkspaceOutlinerTreeItem>()->Export);
+				MenuContext->SelectedItems.Add(Item);
+			}
+		}
+
 		FToolMenuContext Context;
 		Context.AddObject(MenuContext);
 
@@ -216,6 +215,9 @@ TSharedPtr<SWidget> FWorkspaceOutlinerMode::CreateContextMenu()
 void FWorkspaceOutlinerMode::OnItemDoubleClick(FSceneOutlinerTreeItemPtr Item)
 {
 	OpenItems({Item} );
+
+	// Double-clicking automatically expands the item
+	SceneOutliner->SetItemExpansion(Item, true);
 }
 
 FReply FWorkspaceOutlinerMode::OnKeyDown(const FKeyEvent& InKeyEvent)
@@ -228,7 +230,7 @@ FReply FWorkspaceOutlinerMode::OnKeyDown(const FKeyEvent& InKeyEvent)
 	return ISceneOutlinerMode::OnKeyDown(InKeyEvent);
 }
 
-void FWorkspaceOutlinerMode::HandleItemSelection(const FSceneOutlinerItemSelection& Selection)
+void FWorkspaceOutlinerMode::HandleItemSelection(const FSceneOutlinerItemSelection& Selection) const
 {
 	if (Selection.Num() > 0)
 	{
@@ -379,13 +381,13 @@ TUniquePtr<ISceneOutlinerHierarchy> FWorkspaceOutlinerMode::CreateHierarchy()
 	return MakeUnique<FWorkspaceOutlinerHierarchy>(this, WeakWorkspace);
 }
 
-void FWorkspaceOutlinerMode::OnWorkspaceModified(UWorkspace* InWorkspace)
+void FWorkspaceOutlinerMode::OnWorkspaceModified(UWorkspace* InWorkspace) const 
 {
 	ensure(InWorkspace == WeakWorkspace.Get());
 	SceneOutliner->FullRefresh();
 }
 
-void FWorkspaceOutlinerMode::ResetOutlinerSelection()
+void FWorkspaceOutlinerMode::ResetOutlinerSelection() const
 {
 	SceneOutliner->ClearSelection();
 }
@@ -398,17 +400,12 @@ void FWorkspaceOutlinerMode::OpenItems(TArrayView<const FSceneOutlinerTreeItemPt
 		{
 			bool bHandled = false;
 
-			FWorkspaceOutlinerItemExport Export = [&TreeItem]()
-			{
-				const FWorkspaceOutlinerTreeItem* OutlinerTreeItem = TreeItem->CastTo<FWorkspaceOutlinerTreeItem>();
-				return TreeItem->Export;
-			}();
-			
-			if (const TSharedPtr<IWorkspaceOutlinerItemDetails> SharedDetails = FWorkspaceEditorModule::GetOutlinerItemDetails(MakeOutlinerDetailsId(Export)))
+			if (const TSharedPtr<IWorkspaceOutlinerItemDetails> SharedDetails = FWorkspaceEditorModule::GetOutlinerItemDetails(MakeOutlinerDetailsId(TreeItem->Export)))
 			{
 				UWorkspaceItemMenuContext* MenuContext = NewObject<UWorkspaceItemMenuContext>();
-				MenuContext->SelectedExports.Add(Export);
-				
+				MenuContext->SelectedExports.Add(TreeItem->Export);
+				MenuContext->SelectedItems.Add(Item);
+
 				FToolMenuContext Context(MenuContext);
 				WeakWorkspaceEditor.Pin()->InitToolMenuContext(Context);
 				
@@ -426,12 +423,12 @@ void FWorkspaceOutlinerMode::OpenItems(TArrayView<const FSceneOutlinerTreeItemPt
 	}
 }
 
-void FWorkspaceOutlinerMode::Open()
+void FWorkspaceOutlinerMode::Open() const
 {
 	OpenItems(SceneOutliner->GetSelectedItems());
 }
 
-void FWorkspaceOutlinerMode::Delete()
+void FWorkspaceOutlinerMode::Delete() const
 {
 	const FSceneOutlinerItemSelection& Selection = SceneOutliner->GetSelection();
 	if (Selection.Num() == 0)
@@ -524,7 +521,7 @@ bool FWorkspaceOutlinerMode::CanDelete() const
 	return true;
 }
 
-void FWorkspaceOutlinerMode::Rename()
+void FWorkspaceOutlinerMode::Rename() const
 {
 	const FSceneOutlinerItemSelection& Selection = SceneOutliner->GetSelection();
 	if (Selection.Num() == 1)
@@ -550,7 +547,7 @@ bool FWorkspaceOutlinerMode::CanRename() const
 	return false;
 }
 
-void FWorkspaceOutlinerMode::OnAssetRegistryAssetUpdate(const FAssetData& AssetData)
+void FWorkspaceOutlinerMode::OnAssetRegistryAssetUpdate(const FAssetData& AssetData) const
 {
 	SceneOutliner->FullRefresh();
 }
diff --git a/Engine/Plugins/Experimental/Workspace/Source/WorkspaceEditor/Private/Outliner/WorkspaceOutlinerMode.h b/Engine/Plugins/Experimental/Workspace/Source/WorkspaceEditor/Private/Outliner/WorkspaceOutlinerMode.h
--- a/Engine/Plugins/Experimental/Workspace/Source/WorkspaceEditor/Private/Outliner/WorkspaceOutlinerMode.h
+++ b/Engine/Plugins/Experimental/Workspace/Source/WorkspaceEditor/Private/Outliner/WorkspaceOutlinerMode.h
@@ -22,8 +22,8 @@ namespace UE::Workspace
 		virtual void OnItemDoubleClick(FSceneOutlinerTreeItemPtr Item) override;
 		virtual FReply OnKeyDown(const FKeyEvent& InKeyEvent) override;
 		virtual void OnItemSelectionChanged(FSceneOutlinerTreeItemPtr Item, ESelectInfo::Type SelectionType, const FSceneOutlinerItemSelection& Selection) override;
-		virtual bool CanCustomizeToolbar() const { return true; }
-		virtual ESelectionMode::Type GetSelectionMode() const { return ESelectionMode::Multi; }
+		virtual bool CanCustomizeToolbar() const override { return true; }
+		virtual ESelectionMode::Type GetSelectionMode() const override { return ESelectionMode::Multi; }
 		virtual bool CanDelete() const override;
 		virtual bool CanRename() const override;
 		virtual bool CanRenameItem(const ISceneOutlinerTreeItem& Item) const override;
@@ -32,19 +32,18 @@ namespace UE::Workspace
 		virtual TUniquePtr<ISceneOutlinerHierarchy> CreateHierarchy() override;		
 		// End ISceneOutlinerMode overrides
 
-		void HandleItemSelection(const FSceneOutlinerItemSelection& Selection);
-		void OnWorkspaceModified(UWorkspace* InWorkspace);
-		void ResetOutlinerSelection();
+		void HandleItemSelection(const FSceneOutlinerItemSelection& Selection) const;
+		void OnWorkspaceModified(UWorkspace* InWorkspace) const;
+		void ResetOutlinerSelection() const;
 
-		void Open();
-		void Delete();
+		void Open() const;
+		void Delete() const;
 		bool CanDeleteItem(const ISceneOutlinerTreeItem& Item) const;
-		void Rename();
+		void Rename() const;
 
 		void OpenItems(TArrayView<const FSceneOutlinerTreeItemPtr> Items) const;
-		void DeleteItems(TArrayView<const FSceneOutlinerTreeItemPtr> Items) const;
 
-		void OnAssetRegistryAssetUpdate(const FAssetData& AssetData);
+		void OnAssetRegistryAssetUpdate(const FAssetData& AssetData) const;
 	private:
 		TWeakObjectPtr<UWorkspace> WeakWorkspace;
 		TWeakPtr<IWorkspaceEditor> WeakWorkspaceEditor;
