- Change way variable flags are filtered, now works with set of inclusion/exclusion flags to filter against
- Added "Referenced" variable flag, set on recursed SharedVariables their entries
- Filter "Find" variable picker to search scope as well (previously only applied to replace variable picker)
- Changed over .Get<AssetRegistryData> to .GetPtr<AssetRegistryData> to ensure matching data typing

diff --git a/Engine/Plugins/Experimental/AnimNext/Source/AnimNextEditor/Internal/Variables/VariablePickerArgs.h b/Engine/Plugins/Experimental/AnimNext/Source/AnimNextEditor/Internal/Variables/VariablePickerArgs.h
--- a/Engine/Plugins/Experimental/AnimNext/Source/AnimNextEditor/Internal/Variables/VariablePickerArgs.h
+++ b/Engine/Plugins/Experimental/AnimNext/Source/AnimNextEditor/Internal/Variables/VariablePickerArgs.h
@@ -48,8 +48,9 @@ struct FVariablePickerArgs
 	FOnContextSensitivityChanged OnContextSensitivityChanged;
 
 	// Default filter for AnimNext export flags to be used for variable list population
-	EAnimNextExportedVariableFlags FlagFilter = EAnimNextExportedVariableFlags::Declared | EAnimNextExportedVariableFlags::Public;
-
+	EAnimNextExportedVariableFlags FlagInclusionFilter = EAnimNextExportedVariableFlags::Declared | EAnimNextExportedVariableFlags::Public;
+	EAnimNextExportedVariableFlags FlagExclusionFilter = EAnimNextExportedVariableFlags::Referenced;
+	
 	// Whether the search box should be focused on widget creation
 	bool bFocusSearchWidget = true;
 };
diff --git a/Engine/Plugins/Experimental/AnimNext/Source/AnimNextEditor/Private/Common/AnimNextAssetFindReplaceVariables.cpp b/Engine/Plugins/Experimental/AnimNext/Source/AnimNextEditor/Private/Common/AnimNextAssetFindReplaceVariables.cpp
--- a/Engine/Plugins/Experimental/AnimNext/Source/AnimNextEditor/Private/Common/AnimNextAssetFindReplaceVariables.cpp
+++ b/Engine/Plugins/Experimental/AnimNext/Source/AnimNextEditor/Private/Common/AnimNextAssetFindReplaceVariables.cpp
@@ -32,8 +32,10 @@ FString UAnimNextAssetFindReplaceVariables::GetFindResultStringFromAssetData(con
 	{
 		if (SearchReference.GetName() == Export.Identifier)
 		{
-			const FAnimNextVariableReferenceData& ReferenceData = Export.Data.Get<FAnimNextVariableReferenceData>();
-			ReferenceNames.Add(ReferenceData.PinPath);
+			if (const FAnimNextVariableReferenceData* ReferenceData = Export.Data.GetPtr<FAnimNextVariableReferenceData>())
+			{
+				ReferenceNames.Add(ReferenceData->PinPath);
+			}
 		}
 	}
 	
@@ -240,6 +242,68 @@ TSharedRef<SWidget> UAnimNextAssetFindReplaceVariables::MakeFindReplaceWidget()
 {	
 	using namespace UE::AnimNext::Editor;
 
+
+	auto VariableScopeFilter = FOnFilterVariable::CreateLambda([this](const FAnimNextSoftVariableReference& InVariableReference)
+	{
+      	// Filter variable picker to correct scope
+      	bool bShouldFilterAsset = false;
+      	TSharedPtr<UE::Workspace::IWorkspaceEditor> SharedWorkspaceEditor = WeakWorkspaceEditor.Pin();
+      	if (SharedWorkspaceEditor.IsValid())
+      	{
+      		// Include everything within the export chain of the currently opened asset
+      		if (CurrentSearchScope == ESearchScope::Asset)
+      		{
+      			const UE::Workspace::FWorkspaceDocument& Document = SharedWorkspaceEditor->GetFocusedWorkspaceDocument();
+  
+      			TArray<const UAnimNextRigVMAsset*> AssetsInExportChain;
+      			Document.Export.GetAssetsOfType<const UAnimNextRigVMAsset>(AssetsInExportChain, true);
+  
+      			auto ProcessAsset = [&AssetsInExportChain](const UAnimNextRigVMAsset* InAsset)
+      			{
+      				if (InAsset)
+      				{
+      					const UAnimNextRigVMAssetEditorData* EditorData = UE::AnimNext::UncookedOnly::FUtils::GetEditorData<const UAnimNextRigVMAssetEditorData,const UAnimNextRigVMAsset>(InAsset);
+      					EditorData->ForEachEntryOfType<UAnimNextSharedVariablesEntry>([&AssetsInExportChain](const UAnimNextSharedVariablesEntry* VariablesEntry)
+      					{
+      						// [TODO] struct, currently doesn't have any exports so cannot filter against them
+      						if (VariablesEntry->GetAsset())
+      						{
+      							AssetsInExportChain.AddUnique(VariablesEntry->GetAsset());							
+      						}
+  
+      						return true;
+      					});
+      				}
+      			};
+      	
+      			for (int32 AssetIndex = 0; AssetIndex < AssetsInExportChain.Num(); ++AssetIndex)
+      			{
+      				const UAnimNextRigVMAsset* Asset = AssetsInExportChain[AssetIndex];
+      				// This can append to AssetsInExportChain array so it handles recursive shared variable references
+      				ProcessAsset(Asset);
+      			}
+      			
+      			// GetAssetsOfType force loads the assets, so if the variable owner is part of the export chain it would have been loaded		
+      			const UObject* VariableOwner = InVariableReference.GetSoftObjectPath().ResolveObject();
+      			bShouldFilterAsset = !VariableOwner || !AssetsInExportChain.Contains(VariableOwner);
+      		}
+      		else if (CurrentSearchScope == ESearchScope::Workspace)
+      		{
+      			TArray<FAssetData> Assets;
+      			SharedWorkspaceEditor->GetAssets(Assets);
+  
+      			if (!Assets.ContainsByPredicate([ObjectPath = InVariableReference.GetSoftObjectPath()](const FAssetData& AssetData)
+      			{
+      				return AssetData.GetSoftObjectPath() == ObjectPath;
+      			}))
+      			{
+      				bShouldFilterAsset = true;
+      			}
+      		}
+      	}			
+      	return (SearchReference != InVariableReference && !bShouldFilterAsset) ? EFilterVariableResult::Include : EFilterVariableResult::Exclude;		
+    });
+
 	{
 		FVariablePickerArgs SearchArgs;
 		SearchArgs.OnVariablePicked = FOnVariablePicked::CreateLambda([this](const FAnimNextSoftVariableReference& InVariableReference, const FAnimNextParamType& InType)
@@ -255,8 +319,10 @@ TSharedRef<SWidget> UAnimNextAssetFindReplaceVariables::MakeFindReplaceWidget()
 			return EFilterVariableResult::Include;
 		});
 
+		SearchArgs.OnFilterVariable = VariableScopeFilter;
+
 		// Filter to all declared variables (public and private)
-		SearchArgs.FlagFilter = EAnimNextExportedVariableFlags::Declared;
+		SearchArgs.FlagInclusionFilter = EAnimNextExportedVariableFlags::Declared;
 	
 		SearchVariableComboBox = SNew(SVariablePickerCombo)
 			.PickerArgs(SearchArgs)
@@ -272,70 +338,12 @@ TSharedRef<SWidget> UAnimNextAssetFindReplaceVariables::MakeFindReplaceWidget()
 	
 	{
 		FVariablePickerArgs ReplaceArgs;
+		ReplaceArgs.FlagInclusionFilter = EAnimNextExportedVariableFlags::Declared;
 		ReplaceArgs.OnVariablePicked = FOnVariablePicked::CreateLambda([this](const FAnimNextSoftVariableReference& InVariableReference, const FAnimNextParamType& InType)
 		{
 			ReplaceReference = InVariableReference;
 		});
-		ReplaceArgs.OnFilterVariable = FOnFilterVariable::CreateLambda([this](const FAnimNextSoftVariableReference& InVariableReference)
-		{
-			// Filter variable picker to correct scope
-			bool bShouldFilterAsset = false;
-			TSharedPtr<UE::Workspace::IWorkspaceEditor> SharedWorkspaceEditor = WeakWorkspaceEditor.Pin();
-			if (SharedWorkspaceEditor.IsValid())
-			{
-				// Include everything within the export chain of the currently opened asset
-				if (CurrentSearchScope == ESearchScope::Asset)
-				{
-					const UE::Workspace::FWorkspaceDocument& Document = SharedWorkspaceEditor->GetFocusedWorkspaceDocument();
-
-					TArray<const UAnimNextRigVMAsset*> AssetsInExportChain;
-					Document.Export.GetAssetsOfType<const UAnimNextRigVMAsset>(AssetsInExportChain, true);
-
-					auto ProcessAsset = [&AssetsInExportChain](const UAnimNextRigVMAsset* InAsset)
-					{
-						if (InAsset)
-						{
-							const UAnimNextRigVMAssetEditorData* EditorData = UE::AnimNext::UncookedOnly::FUtils::GetEditorData<const UAnimNextRigVMAssetEditorData,const UAnimNextRigVMAsset>(InAsset);
-							EditorData->ForEachEntryOfType<UAnimNextSharedVariablesEntry>([&AssetsInExportChain](const UAnimNextSharedVariablesEntry* VariablesEntry)
-							{
-								// [TODO] struct, currently doesn't have any exports so cannot filter against them
-								if (VariablesEntry->GetAsset())
-								{
-									AssetsInExportChain.AddUnique(VariablesEntry->GetAsset());							
-								}
-
-								return true;
-							});
-						}
-					};
-			
-					for (int32 AssetIndex = 0; AssetIndex < AssetsInExportChain.Num(); ++AssetIndex)
-					{
-						const UAnimNextRigVMAsset* Asset = AssetsInExportChain[AssetIndex];
-						// This can append to AssetsInExportChain array so it handles recursive shared variable references
-						ProcessAsset(Asset);
-					}
-					
-					// GetAssetsOfType force loads the assets, so if the variable owner is part of the export chain it would have been loaded		
-					const UObject* VariableOwner = InVariableReference.GetSoftObjectPath().ResolveObject();
-					bShouldFilterAsset = !VariableOwner || !AssetsInExportChain.Contains(VariableOwner);
-				}
-				else if (CurrentSearchScope == ESearchScope::Workspace)
-				{
-					TArray<FAssetData> Assets;
-					SharedWorkspaceEditor->GetAssets(Assets);
-
-					if (!Assets.ContainsByPredicate([ObjectPath = InVariableReference.GetSoftObjectPath()](const FAssetData& AssetData)
-					{
-						return AssetData.GetSoftObjectPath() == ObjectPath;
-					}))
-					{
-						bShouldFilterAsset = true;
-					}
-				}
-			}			
-			return (SearchReference != InVariableReference && !bShouldFilterAsset) ? EFilterVariableResult::Include : EFilterVariableResult::Exclude;		
-		});
+		ReplaceArgs.OnFilterVariable = VariableScopeFilter;
 		ReplaceArgs.OnFilterVariableType  = FOnFilterVariableType::CreateLambda([this](const FAnimNextParamType& InParamType)-> EFilterVariableResult
 		{
 			if(SearchType.IsValid())
diff --git a/Engine/Plugins/Experimental/AnimNext/Source/AnimNextEditor/Private/Variables/SVariablePicker.cpp b/Engine/Plugins/Experimental/AnimNext/Source/AnimNextEditor/Private/Variables/SVariablePicker.cpp
--- a/Engine/Plugins/Experimental/AnimNext/Source/AnimNextEditor/Private/Variables/SVariablePicker.cpp
+++ b/Engine/Plugins/Experimental/AnimNext/Source/AnimNextEditor/Private/Variables/SVariablePicker.cpp
@@ -97,18 +97,20 @@ void SVariablePicker::RefreshEntries()
 				PropertyDescs.Reserve(ExportPair.Value.Exports.Num());
 				for(const FAnimNextExport& Export : ExportPair.Value.Exports)
 				{
-					if (Export.Data.GetScriptStruct() == FAnimNextVariableDeclarationData::StaticStruct())
+					if (const FAnimNextVariableDeclarationData* VariableEntry = Export.Data.GetPtr<FAnimNextVariableDeclarationData>())
 					{
-						const FAnimNextVariableDeclarationData& VariableEntry = Export.Data.Get<FAnimNextVariableDeclarationData>();
-						if(EnumHasAllFlags(static_cast<EAnimNextExportedVariableFlags>(VariableEntry.Flags), Args.FlagFilter) && Export.Identifier != NAME_None && VariableEntry.Type.IsValid())
+						const bool bHasInclusionFlags = EnumHasAllFlags(static_cast<EAnimNextExportedVariableFlags>(VariableEntry->Flags), Args.FlagInclusionFilter);
+						const bool bHasExclusionFlags = EnumHasAnyFlags(static_cast<EAnimNextExportedVariableFlags>(VariableEntry->Flags), Args.FlagExclusionFilter);
+						
+						if(bHasInclusionFlags && !bHasExclusionFlags && Export.Identifier != NAME_None && VariableEntry->Type.IsValid())
 						{
 							FAnimNextSoftVariableReference VariableReference(Export.Identifier, ExportPair.Key.ToSoftObjectPath());
 
 							if (!Args.OnFilterVariable.IsBound() || Args.OnFilterVariable.Execute(VariableReference) == EFilterVariableResult::Include)
 							{
-								if (!Args.OnFilterVariableType.IsBound() || Args.OnFilterVariableType.Execute(VariableEntry.Type) == EFilterVariableResult::Include)
+								if (!Args.OnFilterVariableType.IsBound() || Args.OnFilterVariableType.Execute(VariableEntry->Type) == EFilterVariableResult::Include)
 								{
-									PropertyDescs.Emplace(Export.Identifier, VariableEntry.Type.GetContainerType(), VariableEntry.Type.GetValueType(), VariableEntry.Type.GetValueTypeObject());
+									PropertyDescs.Emplace(Export.Identifier, VariableEntry->Type.GetContainerType(), VariableEntry->Type.GetValueType(), VariableEntry->Type.GetValueTypeObject());
 								}
 							}
 						}
diff --git a/Engine/Plugins/Experimental/AnimNext/Source/AnimNextUncookedOnly/Internal/UncookedOnlyUtils.h b/Engine/Plugins/Experimental/AnimNext/Source/AnimNextUncookedOnly/Internal/UncookedOnlyUtils.h
--- a/Engine/Plugins/Experimental/AnimNext/Source/AnimNextUncookedOnly/Internal/UncookedOnlyUtils.h
+++ b/Engine/Plugins/Experimental/AnimNext/Source/AnimNextUncookedOnly/Internal/UncookedOnlyUtils.h
@@ -230,7 +230,7 @@ private:
 	static void GetFunctionsWorkspaceExportsRecursive(const UAnimNextRigVMAssetEditorData* EditorData, FWorkspaceOutlinerItemExports& OutExports, const FWorkspaceOutlinerItemExport& RootExport, int32 ParentExportIndex, const TArray<FRigVMGraphFunctionData>& Functions, bool bPublicFunctions);
 	
 	// Gets the exported variables that are used by a RigVM asset	
-	static void GetAssetVariableExports(const UAnimNextRigVMAssetEditorData* EditorData, TSet<FAnimNextExport>& OutExports, FAssetRegistryTagsContext Context);
+	static void GetAssetVariableExports(const UAnimNextRigVMAssetEditorData* EditorData, TSet<FAnimNextExport>& OutExports, FAssetRegistryTagsContext Context, EAnimNextExportedVariableFlags InFlags = EAnimNextExportedVariableFlags::NoFlags);
 
 	static void GetStructVariableExports(const UScriptStruct* Struct, TSet<FAnimNextExport>& OutExports);
 	static void GetSubGraphVariableExportsRecursive(const URigVMEdGraph* RigVMEdGraph, TSet<FAnimNextExport>& OutExports);
diff --git a/Engine/Plugins/Experimental/AnimNext/Source/AnimNextUncookedOnly/Private/UncookedOnlyUtils.cpp b/Engine/Plugins/Experimental/AnimNext/Source/AnimNextUncookedOnly/Private/UncookedOnlyUtils.cpp
--- a/Engine/Plugins/Experimental/AnimNext/Source/AnimNextUncookedOnly/Private/UncookedOnlyUtils.cpp
+++ b/Engine/Plugins/Experimental/AnimNext/Source/AnimNextUncookedOnly/Private/UncookedOnlyUtils.cpp
@@ -735,7 +735,10 @@ FAnimNextParamType FUtils::FindVariableType(const FAnimNextVariableReference& In
 			{
 				if(Export.Identifier == InVariableReference.GetName())
 				{
-					return Export.Data.Get<FAnimNextVariableDeclarationData>().Type;
+					if (const FAnimNextVariableDeclarationData* VariableDeclaration = Export.Data.GetPtr<FAnimNextVariableDeclarationData>())
+					{
+						return VariableDeclaration->Type;
+					}
 				}
 			}
 		}
@@ -775,9 +778,10 @@ FAnimNextParamType FUtils::FindVariableType(const FAnimNextSoftVariableReference
 
 			for(const FAnimNextExport& Export : Exports.Exports)
 			{
-				if(Export.Identifier == InSoftVariableReference.GetName())
+				const FAnimNextVariableDeclarationData* Data = Export.Data.GetPtr<FAnimNextVariableDeclarationData>();
+				if(Export.Identifier == InSoftVariableReference.GetName() && Data)
 				{
-					return Export.Data.Get<FAnimNextVariableDeclarationData>().Type;
+					return Data->Type;
 				}
 			}
 		}
@@ -791,7 +795,10 @@ FAnimNextParamType FUtils::FindVariableType(const FAnimNextSoftVariableReference
 				{
 					if(Export.Identifier == InSoftVariableReference.GetName())
 					{
-						return Export.Data.Get<FAnimNextVariableDeclarationData>().Type;
+						if (const FAnimNextVariableDeclarationData* VariableDeclaration = Export.Data.GetPtr<FAnimNextVariableDeclarationData>())
+						{
+							return VariableDeclaration->Type;
+						}
 					}
 				}
 			}
@@ -860,13 +867,15 @@ static void AddParamToSet(const FAnimNextExport& InNewParam, TSet<FAnimNextExpor
 {
 	if(FAnimNextExport* ExistingEntry = OutExports.Find(InNewParam))
 	{
-		const FAnimNextVariableDeclarationData& NewData = InNewParam.Data.Get<FAnimNextVariableDeclarationData>();
-		FAnimNextVariableDeclarationData& ExistingData = ExistingEntry->Data.GetMutable<FAnimNextVariableDeclarationData>();
-		if(ExistingData.Type != NewData.Type)
+		const FAnimNextVariableDeclarationData* NewData = InNewParam.Data.GetPtr<FAnimNextVariableDeclarationData>();
+		FAnimNextVariableDeclarationData* ExistingData = ExistingEntry->Data.GetMutablePtr<FAnimNextVariableDeclarationData>();
+		check(NewData && ExistingData);
+		
+		if(ExistingData->Type != NewData->Type)
 		{
-			UE_LOGFMT(LogAnimation, Warning, "Type mismatch between parameter {ParameterName}. {ParamType1} vs {ParamType1}", InNewParam.Identifier, NewData.Type.ToString(), ExistingData.Type.ToString());
+			UE_LOGFMT(LogAnimation, Warning, "Type mismatch between parameter {ParameterName}. {ParamType1} vs {ParamType1}", InNewParam.Identifier, NewData->Type.ToString(), ExistingData->Type.ToString());
 		}
-		ExistingData.Flags |= NewData.Flags;
+		ExistingData->Flags |= NewData->Flags;
 	}
 	else
 	{
@@ -883,19 +892,20 @@ void FUtils::GetAssetVariableExports(const UAnimNextRigVMAssetEditorData* Editor
 	OutExports.Exports.Append(ExportSet.Array());
 }
 
-void FUtils::GetAssetVariableExports(const UAnimNextRigVMAssetEditorData* InEditorData, TSet<FAnimNextExport>& OutExports, FAssetRegistryTagsContext Context)
+void FUtils::GetAssetVariableExports(const UAnimNextRigVMAssetEditorData* InEditorData, TSet<FAnimNextExport>& OutExports, FAssetRegistryTagsContext Context, EAnimNextExportedVariableFlags InFlags /*= EAnimNextExportedVariableFlags::NoFlags*/)
 {
 	for(const UAnimNextRigVMAssetEntry* Entry : InEditorData->Entries)
 	{
 		if(const IAnimNextRigVMExportInterface* ExportInterface = Cast<IAnimNextRigVMExportInterface>(Entry))
 		{
-			EAnimNextExportedVariableFlags Flags = EAnimNextExportedVariableFlags::Declared;
+			EAnimNextExportedVariableFlags Flags = EAnimNextExportedVariableFlags::Declared | InFlags;
 			if(ExportInterface->GetExportAccessSpecifier() == EAnimNextExportAccessSpecifier::Public)
 			{
-				Flags |= EAnimNextExportedVariableFlags::Public;			
-				FAnimNextExport ParameterExport = FAnimNextExport::MakeExport<FAnimNextVariableDeclarationData>(ExportInterface->GetExportName(), ExportInterface->GetExportType(), Flags);
-				AddParamToSet(ParameterExport, OutExports);
+				Flags |= EAnimNextExportedVariableFlags::Public;
 			}
+						
+			FAnimNextExport ParameterExport = FAnimNextExport::MakeExport<FAnimNextVariableDeclarationData>(ExportInterface->GetExportName(), ExportInterface->GetExportType(), Flags);
+			AddParamToSet(ParameterExport, OutExports);
 		}
 		else if(const UAnimNextSharedVariablesEntry* DataInterfaceEntry = Cast<UAnimNextSharedVariablesEntry>(Entry))
 		{
@@ -905,7 +915,7 @@ void FUtils::GetAssetVariableExports(const UAnimNextRigVMAssetEditorData* InEdit
 				if(const UAnimNextSharedVariables* SharedVariablesAsset = DataInterfaceEntry->Asset)
 				{
 					UAnimNextSharedVariables_EditorData* EditorData = GetEditorData<UAnimNextSharedVariables_EditorData>(SharedVariablesAsset);
-					GetAssetVariableExports(EditorData, OutExports, Context);
+					GetAssetVariableExports(EditorData, OutExports, Context, EAnimNextExportedVariableFlags::Referenced);
 				}
 				break;
 			case EAnimNextSharedVariablesType::Struct:
@@ -1526,7 +1536,7 @@ void FUtils::MoveVariableToAsset(UAnimNextVariableEntry* VariableEntry, UAnimNex
 						
 			FString DefaultValueString;
 			VariableEntry->GetDefaultValueString(DefaultValueString);
-			
+
 			check(NewEditorData->AddVariable(VariableEntry->GetVariableName(), VariableEntry->GetType(), DefaultValueString, bSetupUndoRedo, bPrintPythonCommands));
 
 			const FAnimNextSoftVariableReference FindReference(VariableEntry->GetEntryName(), CurrentOuter);
diff --git a/Engine/Plugins/Experimental/AnimNext/Source/AnimNextUncookedOnly/Public/AnimNextExports.h b/Engine/Plugins/Experimental/AnimNext/Source/AnimNextUncookedOnly/Public/AnimNextExports.h
--- a/Engine/Plugins/Experimental/AnimNext/Source/AnimNextUncookedOnly/Public/AnimNextExports.h
+++ b/Engine/Plugins/Experimental/AnimNext/Source/AnimNextUncookedOnly/Public/AnimNextExports.h
@@ -21,6 +21,7 @@ enum class EAnimNextExportedVariableFlags : uint32
 	Read = 0x02,
 	Write = 0x04,
 	Declared = 0x08,
+	Referenced = 0x10,
 	Max
 };
 
diff --git a/Engine/Plugins/Experimental/AnimNextAnimGraph/Source/AnimNextAnimGraphEditor/Private/Graph/AnimNextAnimationGraphMenuExtensions.cpp b/Engine/Plugins/Experimental/AnimNextAnimGraph/Source/AnimNextAnimGraphEditor/Private/Graph/AnimNextAnimationGraphMenuExtensions.cpp
--- a/Engine/Plugins/Experimental/AnimNextAnimGraph/Source/AnimNextAnimGraphEditor/Private/Graph/AnimNextAnimationGraphMenuExtensions.cpp
+++ b/Engine/Plugins/Experimental/AnimNextAnimGraph/Source/AnimNextAnimGraphEditor/Private/Graph/AnimNextAnimationGraphMenuExtensions.cpp
@@ -203,13 +203,16 @@ void FAnimationGraphMenuExtensions::RegisterMenus()
 					{
 						// Filter to only show assets with public variables
 						FAnimNextAssetRegistryExports Exports;
-						if(UE::AnimNext::UncookedOnly::FUtils::GetExportedVariablesForAsset(InAssetData, Exports))
+						if(UncookedOnly::FUtils::GetExportedVariablesForAsset(InAssetData, Exports))
 						{
 							for(const FAnimNextExport& Export : Exports.Exports)
 							{
-								if((static_cast<EAnimNextExportedVariableFlags>(Export.Data.Get<FAnimNextVariableDeclarationData>().Flags) & EAnimNextExportedVariableFlags::Public) != EAnimNextExportedVariableFlags::NoFlags)
+								if (const FAnimNextVariableDeclarationData* VariableDeclaration = Export.Data.GetPtr<FAnimNextVariableDeclarationData>())
 								{
-									return false;
+									if (EnumHasAnyFlags(static_cast<EAnimNextExportedVariableFlags>(VariableDeclaration->Flags),EAnimNextExportedVariableFlags::Public))
+									{
+										return false;
+									}					
 								}
 							}
 						}
