Make sure SharedVariables entries actually resolve to a valid object/struct before adding to variable hierarchy (previously could show deleted, null, assets)

diff --git a/Engine/Plugins/Experimental/AnimNext/Source/AnimNextEditor/Private/Variables/Outliner/VariablesOutlinerHierarchy.cpp b/Engine/Plugins/Experimental/AnimNext/Source/AnimNextEditor/Private/Variables/Outliner/VariablesOutlinerHierarchy.cpp
--- a/Engine/Plugins/Experimental/AnimNext/Source/AnimNextEditor/Private/Variables/Outliner/VariablesOutlinerHierarchy.cpp
+++ b/Engine/Plugins/Experimental/AnimNext/Source/AnimNextEditor/Private/Variables/Outliner/VariablesOutlinerHierarchy.cpp
@@ -133,41 +133,45 @@ void FVariablesOutlinerHierarchy::ProcessAsset(const TSoftObjectPtr<UAnimNextRig
 		{
 		case EAnimNextSharedVariablesType::Asset:
 			{
-				const UAnimNextSharedVariables* SharedVariables = InSharedVariablesEntry->GetAsset();
-				OutSharedVariableSourcePaths.Add(SharedVariables);
+				if (const UAnimNextSharedVariables* SharedVariables = InSharedVariablesEntry->GetAsset())
+				{
+					OutSharedVariableSourcePaths.Add(SharedVariables);
 				
-				// Would have expected shared variable (assets) to be part of the outer export chain, so should have been handled as regular assets.
-				if (!Outliner->Assets.Contains(SharedVariables))
-				{	
-					TArray<const FSoftObjectPath> SharedVariableSourcePaths;
-					ProcessAsset( const_cast<UAnimNextSharedVariables*>(SharedVariables), Depth - 1, ParentItems, SharedVariableSourcePaths);
-				}
+					// Would have expected shared variable (assets) to be part of the outer export chain, so should have been handled as regular assets.
+					if (!Outliner->Assets.Contains(SharedVariables))
+					{	
+						TArray<const FSoftObjectPath> SharedVariableSourcePaths;
+						ProcessAsset( const_cast<UAnimNextSharedVariables*>(SharedVariables), Depth - 1, ParentItems, SharedVariableSourcePaths);
+					}
+				}				
 			}
 			break;
 
 		case EAnimNextSharedVariablesType::Struct:
-			{
-				if (FSceneOutlinerTreeItemPtr Item = Mode->CreateItemFor<FVariablesOutlinerStructSharedVariablesItem>(InSharedVariablesEntry))
+			{				
+				if (const UScriptStruct* Struct = InSharedVariablesEntry->GetStruct())
 				{
-					OutItems.Add(Item);
-				}
-
-				UE_CLOG(GOutputHierarchyDebugInformationCvar.GetValueOnAnyThread(), LogAnimation, Warning, TEXT("~ FVariablesOutlinerStructSharedVariablesItem: %s %i"), *InSharedVariablesEntry->GetObjectPath().ToString(), Depth);
-				
-				const UScriptStruct* Struct = InSharedVariablesEntry->GetStruct();
-				OutSharedVariableSourcePaths.Add(Struct);
-				int32 Index = 0;
-				for (TFieldIterator<FProperty> It(Struct); It; ++It)
-				{
-					const FProperty* Property = *It;
-					if (FSceneOutlinerTreeItemPtr Item = Mode->CreateItemFor<FVariablesOutlinerEntryItem>(Property))
+					if (FSceneOutlinerTreeItemPtr Item = Mode->CreateItemFor<FVariablesOutlinerStructSharedVariablesItem>(InSharedVariablesEntry))
 					{
-						FVariablesOutlinerEntryItem* EntryItem = Item->CastTo<FVariablesOutlinerEntryItem>();
-						EntryItem->WeakSharedVariablesEntry = InSharedVariablesEntry;
-						EntryItem->SortValue = Index++;
 						OutItems.Add(Item);
 					}
-				}
+					
+					UE_CLOG(GOutputHierarchyDebugInformationCvar.GetValueOnAnyThread(), LogAnimation, Warning, TEXT("~ FVariablesOutlinerStructSharedVariablesItem: %s %i"), *InSharedVariablesEntry->GetObjectPath().ToString(), Depth);
+					
+					OutSharedVariableSourcePaths.Add(Struct);
+					int32 Index = 0;
+					for (TFieldIterator<FProperty> It(Struct); It; ++It)
+					{
+						const FProperty* Property = *It;
+						if (FSceneOutlinerTreeItemPtr Item = Mode->CreateItemFor<FVariablesOutlinerEntryItem>(Property))
+						{
+							FVariablesOutlinerEntryItem* EntryItem = Item->CastTo<FVariablesOutlinerEntryItem>();
+							EntryItem->WeakSharedVariablesEntry = InSharedVariablesEntry;
+							EntryItem->SortValue = Index++;
+							OutItems.Add(Item);
+						}
+					}
+				}				
 			}
 			break;
 		}		
