Fixed-up incorrect delegate being returned and registered to

diff --git a/Engine/Plugins/Experimental/AnimNext/Source/AnimNext/Private/AnimNextRigVMAsset.cpp b/Engine/Plugins/Experimental/AnimNext/Source/AnimNext/Private/AnimNextRigVMAsset.cpp
--- a/Engine/Plugins/Experimental/AnimNext/Source/AnimNext/Private/AnimNextRigVMAsset.cpp
+++ b/Engine/Plugins/Experimental/AnimNext/Source/AnimNext/Private/AnimNextRigVMAsset.cpp
@@ -203,7 +203,7 @@ UAnimNextRigVMAsset::FOnCompileJobEvent& UAnimNextRigVMAsset::OnCompileJobStarte
 
 UAnimNextRigVMAsset::FOnCompileJobEvent& UAnimNextRigVMAsset::OnCompileJobFinished()
 {
-	return UE::AnimNext::Private::GOnCompileJobStarted;
+	return UE::AnimNext::Private::GOnCompileJobFinished;
 }
 
 #endif
\ No newline at end of file
diff --git a/Engine/Plugins/Experimental/AnimNext/Source/AnimNextUncookedOnly/Private/AnimNextScopedCompileJob.cpp b/Engine/Plugins/Experimental/AnimNext/Source/AnimNextUncookedOnly/Private/AnimNextScopedCompileJob.cpp
--- a/Engine/Plugins/Experimental/AnimNext/Source/AnimNextUncookedOnly/Private/AnimNextScopedCompileJob.cpp
+++ b/Engine/Plugins/Experimental/AnimNext/Source/AnimNextUncookedOnly/Private/AnimNextScopedCompileJob.cpp
@@ -145,7 +145,7 @@ FScopedCompileJob::~FScopedCompileJob()
 
 				if (!EditorData->bSuspendCompilationNotifications)
 				{
-					UAnimNextRigVMAsset::OnCompileJobStarted().Broadcast(Asset);
+					UAnimNextRigVMAsset::OnCompileJobFinished().Broadcast(Asset);
 					EditorData->OnCompileJobFinished();
 				}
 			}
