Fix shipping compile error

diff --git a/Engine/Plugins/Experimental/AnimNext/Source/AnimNext/Public/UAFAssetInstance.h b/Engine/Plugins/Experimental/AnimNext/Source/AnimNext/Public/UAFAssetInstance.h
--- a/Engine/Plugins/Experimental/AnimNext/Source/AnimNext/Public/UAFAssetInstance.h
+++ b/Engine/Plugins/Experimental/AnimNext/Source/AnimNext/Public/UAFAssetInstance.h
@@ -234,12 +234,14 @@ struct FUAFAssetInstance
 		return HostInstance;
 	}
 
-#if ANIMNEXT_TRACE_ENABLED
 	uint64 GetUniqueId() const
 	{
+#if ANIMNEXT_TRACE_ENABLED
 		return UniqueId;
-	}
+#else
+		return 0;
 #endif 
+	}
 
 #if DO_CHECK
 	// Validate that the supplied property bag is of the same layout as this instance's variables
