Fix off by one insertion index when adding Entry to AnimNextEditorData

diff --git a/Engine/Plugins/Experimental/AnimNext/Source/AnimNextUncookedOnly/Private/AnimNextRigVMAssetEditorData.cpp b/Engine/Plugins/Experimental/AnimNext/Source/AnimNextUncookedOnly/Private/AnimNextRigVMAssetEditorData.cpp
--- a/Engine/Plugins/Experimental/AnimNext/Source/AnimNextUncookedOnly/Private/AnimNextRigVMAssetEditorData.cpp
+++ b/Engine/Plugins/Experimental/AnimNext/Source/AnimNextUncookedOnly/Private/AnimNextRigVMAssetEditorData.cpp
@@ -2705,7 +2705,7 @@ void UAnimNextRigVMAssetEditorData::AddReferencedObjects(UObject* InThis, FRefer
 
 void UAnimNextRigVMAssetEditorData::AddEntryInternal(UAnimNextRigVMAssetEntry* InEntry)
 {
-	InsertEntryInternal(InEntry, FMath::Max(Entries.Num() - 1, 0));	
+	InsertEntryInternal(InEntry, Entries.Num());	
 }
 
 void UAnimNextRigVMAssetEditorData::InsertEntryInternal(UAnimNextRigVMAssetEntry* InEntry, int32 InsertionIndex)
