Fixup empty tooltip LOCTEXT, replaced with commandlist entry data

diff --git a/Engine/Plugins/Experimental/AnimNext/Source/AnimNextEditor/Private/AnimNextEditorModule.cpp b/Engine/Plugins/Experimental/AnimNext/Source/AnimNextEditor/Private/AnimNextEditorModule.cpp
--- a/Engine/Plugins/Experimental/AnimNext/Source/AnimNextEditor/Private/AnimNextEditorModule.cpp
+++ b/Engine/Plugins/Experimental/AnimNext/Source/AnimNextEditor/Private/AnimNextEditorModule.cpp
@@ -261,12 +261,12 @@ void FAnimNextEditorModule::StartupModule()
 												}
 											}
 										};
-										
-										 InSubmenu->AddMenuEntry(NAME_None, FToolMenuEntry::InitMenuEntry("FindGlobalReferences", LOCTEXT("GlobalScopeLabel", "In Project"), LOCTEXT("", ""),UAnimNextAssetFindReplaceVariables::GetIconFromSearchScope(ESearchScope::Global), FToolUIAction(FToolMenuExecuteAction::CreateLambda(InvokeFindVariables, ESearchScope::Global))));
 
-										InSubmenu->AddMenuEntry(NAME_None, FToolMenuEntry::InitMenuEntry("FindWorkspaceReferences", LOCTEXT("WorkspaceScopeLabel", "In Workspace"), LOCTEXT("", ""),UAnimNextAssetFindReplaceVariables::GetIconFromSearchScope(ESearchScope::Workspace), FToolUIAction(FToolMenuExecuteAction::CreateLambda(InvokeFindVariables, ESearchScope::Workspace))));
+										InSubmenu->AddMenuEntry(NAME_None, FToolMenuEntry::InitMenuEntry("FindGlobalReferences", FVariablesOutlinerCommands::Get().FindReferences->GetLabel(), FVariablesOutlinerCommands::Get().FindReferences->GetDescription(), UAnimNextAssetFindReplaceVariables::GetIconFromSearchScope(ESearchScope::Global), FToolUIAction(FToolMenuExecuteAction::CreateLambda(InvokeFindVariables, ESearchScope::Global))));
 
-										InSubmenu->AddMenuEntry(NAME_None, FToolMenuEntry::InitMenuEntry("FindAssetReferences", LOCTEXT("AssetScopeLabel", "In Asset"), LOCTEXT("", ""),UAnimNextAssetFindReplaceVariables::GetIconFromSearchScope(ESearchScope::Asset), FToolUIAction(FToolMenuExecuteAction::CreateLambda(InvokeFindVariables, ESearchScope::Asset))));		
+										InSubmenu->AddMenuEntry(NAME_None, FToolMenuEntry::InitMenuEntry("FindWorkspaceReferences", FVariablesOutlinerCommands::Get().FindReferencesInWorkspace->GetLabel(), FVariablesOutlinerCommands::Get().FindReferencesInWorkspace->GetDescription(), UAnimNextAssetFindReplaceVariables::GetIconFromSearchScope(ESearchScope::Workspace), FToolUIAction(FToolMenuExecuteAction::CreateLambda(InvokeFindVariables, ESearchScope::Workspace))));
+
+										InSubmenu->AddMenuEntry(NAME_None, FToolMenuEntry::InitMenuEntry("FindAssetReferences", FVariablesOutlinerCommands::Get().FindReferencesInAsset->GetLabel(), FVariablesOutlinerCommands::Get().FindReferencesInAsset->GetDescription(), UAnimNextAssetFindReplaceVariables::GetIconFromSearchScope(ESearchScope::Asset), FToolUIAction(FToolMenuExecuteAction::CreateLambda(InvokeFindVariables, ESearchScope::Asset))));
 									}),
 									FToolUIActionChoice(),
 									EUserInterfaceActionType::Button,
