Removed AddVariables Dialog, replaced with submenu on in-line add button(s) and outliner context menu

diff --git a/Engine/Plugins/Experimental/AnimNext/Source/AnimNextEditor/Internal/Common/AnimNextAssetItemDetails.cpp b/Engine/Plugins/Experimental/AnimNext/Source/AnimNextEditor/Internal/Common/AnimNextAssetItemDetails.cpp
--- a/Engine/Plugins/Experimental/AnimNext/Source/AnimNextEditor/Internal/Common/AnimNextAssetItemDetails.cpp
+++ b/Engine/Plugins/Experimental/AnimNext/Source/AnimNextEditor/Internal/Common/AnimNextAssetItemDetails.cpp
@@ -25,7 +25,6 @@
 #include "Module/RigUnit_AnimNextModuleEvents.h"
 #include "RigVMFunctions/Execution/RigVMFunction_UserDefinedEvent.h"
 #include "UObject/UObjectIterator.h"
-#include "Variables/SAddVariablesDialog.h"
 #include "Widgets/Input/STextEntryPopup.h"
 #include "Variables/SVariablesView.h"
 
@@ -278,74 +277,6 @@ void FAnimNextAssetItemDetails::RegisterToolMenuExtensions()
 					false,
 					FSlateIconFinder::FindIconForClass(SubEntryClass, "ClassIcon.Object"));
 			}
-			else
-			{
-				AssetSection.AddMenuEntry(
-					SubEntryClass->GetFName(),
-					FText::Format(AddEntryLabelFormat, SubEntryClass->GetDisplayNameText()),
-					FText::Format(AddEntryTooltipFormat, SubEntryClass->GetDisplayNameText()),
-					FSlateIconFinder::FindIconForClass(SubEntryClass, "ClassIcon.Object"),
-					FUIAction(
-						FExecuteAction::CreateWeakLambda(WorkspaceItemContext, [&OutlinerData, SubEntryClass, EditorData]()
-						{
-							if(SubEntryClass == UAnimNextVariableEntry::StaticClass())
-							{
-								TSharedRef<SAddVariablesDialog> AddVariablesDialog =
-									SNew(SAddVariablesDialog, TArray<UAnimNextRigVMAssetEditorData*>({ UncookedOnly::FUtils::GetEditorData<UAnimNextRigVMAssetEditorData>(OutlinerData.GetAsset()) }));
-
-								TArray<FVariableToAdd> VariablesToAdd;
-								TArray<FSharedVariablesToAdd> AllSharedVariablesToAdd;
-								if(AddVariablesDialog->ShowModal(VariablesToAdd, AllSharedVariablesToAdd))
-								{
-									FScopedTransaction Transaction(LOCTEXT("AddVariables", "Add Variable(s)"));
-									for (const FVariableToAdd& VariableToAdd : VariablesToAdd)
-									{
-										EditorData->AddVariable(VariableToAdd.Name, VariableToAdd.Type);
-									}
-									for (const FSharedVariablesToAdd& SharedVariablesToAdd : AllSharedVariablesToAdd)
-									{
-										if (SharedVariablesToAdd.SharedVariables.IsType<UAnimNextSharedVariables*>())
-										{
-											EditorData->AddSharedVariables(SharedVariablesToAdd.SharedVariables.Get<UAnimNextSharedVariables*>());
-										}
-										else
-										{
-											EditorData->AddSharedVariablesStruct(SharedVariablesToAdd.SharedVariables.Get<const UScriptStruct*>());
-										}
-									}
-								}
-							}
-							else if(SubEntryClass == UAnimNextSharedVariablesEntry::StaticClass())
-							{
-								TSharedRef<SAddVariablesDialog> AddVariablesDialog =
-									SNew(SAddVariablesDialog, TArray<UAnimNextRigVMAssetEditorData*>({ UncookedOnly::FUtils::GetEditorData<UAnimNextRigVMAssetEditorData>(OutlinerData.GetAsset()) }))
-									.ShouldAddInitialVariable(false);
-
-								TArray<FVariableToAdd> VariablesToAdd;
-								TArray<FSharedVariablesToAdd> AllSharedVariablesToAdd;
-								if(AddVariablesDialog->ShowModal(VariablesToAdd, AllSharedVariablesToAdd))
-								{
-									FScopedTransaction Transaction(LOCTEXT("AddVariables", "Add Variable(s)"));
-									for (const FVariableToAdd& VariableToAdd : VariablesToAdd)
-									{
-										EditorData->AddVariable(VariableToAdd.Name, VariableToAdd.Type);
-									}
-									for (const FSharedVariablesToAdd& SharedVariablesToAdd : AllSharedVariablesToAdd)
-									{
-										if (SharedVariablesToAdd.SharedVariables.IsType<UAnimNextSharedVariables*>())
-										{
-											EditorData->AddSharedVariables(SharedVariablesToAdd.SharedVariables.Get<UAnimNextSharedVariables*>());
-										}
-										else
-										{
-											EditorData->AddSharedVariablesStruct(SharedVariablesToAdd.SharedVariables.Get<const UScriptStruct*>());
-										}
-									}
-								}
-							}
-						})
-					));
-			}
 		}
 
 		AssetSection.AddMenuEntry(
diff --git a/Engine/Plugins/Experimental/AnimNext/Source/AnimNextEditor/Internal/Variables/VariablePickerArgs.h b/Engine/Plugins/Experimental/AnimNext/Source/AnimNextEditor/Internal/Variables/VariablePickerArgs.h
--- a/Engine/Plugins/Experimental/AnimNext/Source/AnimNextEditor/Internal/Variables/VariablePickerArgs.h
+++ b/Engine/Plugins/Experimental/AnimNext/Source/AnimNextEditor/Internal/Variables/VariablePickerArgs.h
@@ -4,7 +4,6 @@
 
 #include "AnimNextExports.h"
 #include "Delegates/Delegate.h"
-#include "Variables/SAddVariablesDialog.h"
 
 struct FAnimNextVariableReference;
 struct FAnimNextSoftVariableReference;
@@ -13,6 +12,16 @@ struct FAnimNextParamType;
 namespace UE::AnimNext::Editor
 {
 
+// Result of a filter operation via FOnFilterVariableType
+enum class EFilterVariableResult : int32
+{
+	Include,
+	Exclude
+};
+
+// Delegate called to filter variables by type for display to the user
+using FOnFilterVariableType = TDelegate<EFilterVariableResult(const FAnimNextParamType& /*InType*/)>;
+
 // Delegate called when a variable has been picked. Graph argument is invalid when an unbound variable is chosen.
 using FOnVariablePicked = TDelegate<void(const FAnimNextSoftVariableReference& /*InSoftVariableReference*/, const FAnimNextParamType& /*InType*/)>;
 
diff --git a/Engine/Plugins/Experimental/AnimNext/Source/AnimNextEditor/Private/AnimNextEditorModule.cpp b/Engine/Plugins/Experimental/AnimNext/Source/AnimNextEditor/Private/AnimNextEditorModule.cpp
--- a/Engine/Plugins/Experimental/AnimNext/Source/AnimNextEditor/Private/AnimNextEditorModule.cpp
+++ b/Engine/Plugins/Experimental/AnimNext/Source/AnimNextEditor/Private/AnimNextEditorModule.cpp
@@ -32,7 +32,6 @@
 #include "Param/ObjectCastLocatorEditor.h"
 #include "Param/ObjectFunctionLocatorEditor.h"
 #include "Param/ObjectPropertyLocatorEditor.h"
-#include "Variables/SAddVariablesDialog.h"
 #include "AnimNextAssetWorkspaceAssetUserData.h"
 #include "AnimNextEditorContext.h"
 #include "AssetCompilationHandler.h"
diff --git a/Engine/Plugins/Experimental/AnimNext/Source/AnimNextEditor/Private/Variables/AddVariableDialogMenuContext.h b/Engine/Plugins/Experimental/AnimNext/Source/AnimNextEditor/Private/Variables/AddVariableDialogMenuContext.h
deleted file mode 100644
diff --git a/Engine/Plugins/Experimental/AnimNext/Source/AnimNextEditor/Private/Variables/Outliner/VariablesOutlinerAssetItem.cpp b/Engine/Plugins/Experimental/AnimNext/Source/AnimNextEditor/Private/Variables/Outliner/VariablesOutlinerAssetItem.cpp
--- a/Engine/Plugins/Experimental/AnimNext/Source/AnimNextEditor/Private/Variables/Outliner/VariablesOutlinerAssetItem.cpp
+++ b/Engine/Plugins/Experimental/AnimNext/Source/AnimNextEditor/Private/Variables/Outliner/VariablesOutlinerAssetItem.cpp
@@ -20,6 +20,7 @@
 #include "Styling/SlateIconFinder.h"
 #include "Styling/StarshipCoreStyle.h"
 #include "UObject/Package.h"
+#include "Variables/AnimNextVariableItemMenuContext.h"
 #include "Variables/SVariablesView.h"
 #include "Widgets/SCompoundWidget.h"
 #include "Widgets/DeclarativeSyntaxSupport.h"
@@ -87,29 +88,49 @@ class SVariablesOutlinerAssetLabel : FSceneOutlinerCommonLabelData, public SComp
 			.Padding(2.0f, 0.0f)
 			.AutoWidth()
 			[
-				SNew(SButton)
-				.ButtonStyle(FAppStyle::Get(), "SimpleButton")
-				.OnClicked_Lambda([WeakLabel]()
-				{
-					if (TSharedPtr<SVariablesOutlinerAssetLabel> SharedThis = WeakLabel.Pin())
+				SNew(SBox)
+				.Visibility(this, &SVariablesOutlinerAssetLabel::GetAddVariableButtonVisibility)
+				[
+					SNew(SComboButton)
+					.ContentPadding(FMargin(0, 0))
+					.ButtonStyle(FAppStyle::Get(), "SimpleButton")
+					.HasDownArrow(false)
+					.ButtonContent()
+					[
+						SNew(SImage)
+						.Image(FAppStyle::Get().GetBrush("Icons.PlusCircle"))
+						.ColorAndOpacity(FSlateColor::UseForeground())
+					]
+					.OnGetMenuContent_Lambda([WeakOutliner = SceneOutliner.AsShared().ToWeakPtr(), WeakItem = InTreeItem.AsShared().ToWeakPtr()]() -> TSharedRef<SWidget>
 					{
-						if (TSharedPtr<FVariablesOutlinerAssetItem> TreeItem = SharedThis->TreeItem.Pin())
+						TSharedPtr<SVariablesOutliner> SharedOutliner = StaticCastSharedPtr<SVariablesOutliner>(WeakOutliner.Pin());
+						ensure(SharedOutliner.IsValid());
+
+						if (const FVariablesOutlinerMode* VariablesMode = static_cast<const FVariablesOutlinerMode*>(SharedOutliner->GetMode()))
 						{
-							if (TSharedPtr<ISceneOutliner> Outliner = SharedThis->WeakSceneOutliner.Pin())
+							FToolMenuContext Context;
+							UAnimNextVariableItemMenuContext* MenuContext = NewObject<UAnimNextVariableItemMenuContext>();
+							MenuContext->WeakOutliner = SharedOutliner;	
+
+							TSharedPtr<FVariablesOutlinerAssetItem> SharedAssetItem = StaticCastSharedPtr<FVariablesOutlinerAssetItem>(WeakItem.Pin());
+							if (SharedAssetItem.IsValid())
 							{
-								FVariablesOutlinerMode::AddNewVariable(TreeItem->SoftAsset.Get());							
+								if (UAnimNextRigVMAsset* Asset = SharedAssetItem->SoftAsset.Get())
+								{
+									if (UAnimNextRigVMAssetEditorData* EditorData = UncookedOnly::FUtils::GetEditorData(Asset))
+									{
+										MenuContext->WeakEditorDatas.Add(EditorData);
+									}
+								}
 							}
+
+							Context.AddObject(MenuContext);
+							static const FName MenuName("VariablesOutliner.AddVariablesMenu");
+							return UToolMenus::Get()->GenerateWidget(MenuName, Context);
 						}
-					}
 
-					return FReply::Handled();
-				})
-				.Visibility(this, &SVariablesOutlinerAssetLabel::GetAddVariableButtonVisibility)
-				.ContentPadding(FMargin(1, 0))
-				[
-					SNew(SImage)
-					.Image(FAppStyle::Get().GetBrush("Icons.PlusCircle"))
-					.ColorAndOpacity(FSlateColor::UseForeground())
+						return SNullWidget::NullWidget;
+					})
 				]
 			]
 			+SHorizontalBox::Slot()
diff --git a/Engine/Plugins/Experimental/AnimNext/Source/AnimNextEditor/Private/Variables/Outliner/VariablesOutlinerMode.cpp b/Engine/Plugins/Experimental/AnimNext/Source/AnimNextEditor/Private/Variables/Outliner/VariablesOutlinerMode.cpp
--- a/Engine/Plugins/Experimental/AnimNext/Source/AnimNextEditor/Private/Variables/Outliner/VariablesOutlinerMode.cpp
+++ b/Engine/Plugins/Experimental/AnimNext/Source/AnimNextEditor/Private/Variables/Outliner/VariablesOutlinerMode.cpp
@@ -10,15 +10,18 @@
 #include "AnimNextRigVMAssetEditorData.h"
 #include "IWorkspaceEditor.h"
 #include "AnimNextEditorModule.h"
+#include "AnimNextEditorStyle.h"
 #include "AssetRegistry/AssetRegistryModule.h"
 #include "ScopedTransaction.h"
 #include "ToolMenus.h"
 #include "Framework/Commands/GenericCommands.h"
 #include "Entries/AnimNextRigVMAssetEntry.h"
-#include "Variables/SAddVariablesDialog.h"
 #include "Variables/AnimNextVariableItemMenuContext.h"
 #include "AnimNextRigVMAsset.h"
+#include "ContentBrowserModule.h"
 #include "FileHelpers.h"
+#include "StructViewerFilter.h"
+#include "StructViewerModule.h"
 #include "VariablesOutlinerAssetItem.h"
 #include "VariablesOutlinerCommands.h"
 #include "VariablesOutlinerStructSharedVariablesItem.h"
@@ -29,6 +32,7 @@
 #include "Entries/AnimNextVariableEntry.h"
 #include "Variables/SVariablesView.h"
 #include "Entries/AnimNextSharedVariablesEntry.h"
+#include "Styling/SlateIconFinder.h"
 #include "Variables/AnimNextSharedVariables.h"
 #include "Variables/AnimNextVariableSettings.h"
 #include "Widgets/Docking/SDockTab.h"
@@ -43,6 +47,21 @@ FVariablesOutlinerMode::FVariablesOutlinerMode(SVariablesOutliner* InVariablesOu
 	, WeakWorkspaceEditor(InWorkspaceEditor)
 {
 	CommandList = MakeShared<FUICommandList>();
+
+	static const FName MenuName("VariablesOutliner.AddVariablesMenu");
+	UToolMenus* ToolMenus = UToolMenus::Get();
+	if (!ToolMenus->IsMenuRegistered(MenuName))
+	{
+		FToolMenuOwnerScoped ToolMenuOwnerScope(this);
+
+		if (UToolMenu* Menu = ToolMenus->RegisterMenu(MenuName))
+		{
+			Menu->AddDynamicSection(TEXT("Variables"), FNewToolMenuDelegate::CreateLambda([](UToolMenu* InMenu)
+			{
+				PopulateNewVariableToolMenuEntries(InMenu, false);				
+			}));
+		}
+	}
 }
 
 void FVariablesOutlinerMode::Rebuild()
@@ -73,15 +92,16 @@ TSharedPtr<SWidget> FVariablesOutlinerMode::CreateContextMenu()
 				/** Asset specific actions */
 				FToolMenuSection& AssetSection = InMenu->AddSection("Asset", LOCTEXT("AssetSectionLabel", "Asset"));
 				{
-					AssetSection.AddMenuEntryWithCommandList(FVariablesOutlinerCommands::Get().AddNewVariable, MenuContext->WeakCommandList.Pin(), TAttribute<FText>(), TAttribute<FText>(), FSlateIcon(FAppStyle::GetAppStyleSetName(), TEXT("Icons.Plus")));
-					AssetSection.AddMenuEntryWithCommandList(FVariablesOutlinerCommands::Get().AddNewVariables, MenuContext->WeakCommandList.Pin(), TAttribute<FText>(), TAttribute<FText>(), FSlateIcon(FAppStyle::GetAppStyleSetName(), TEXT("Icons.Plus")));
 					AssetSection.AddMenuEntryWithCommandList(FVariablesOutlinerCommands::Get().SaveAsset, MenuContext->WeakCommandList.Pin(), TAttribute<FText>(), TAttribute<FText>(), FSlateIcon(FAppStyle::Get().GetStyleSetName(), "AssetEditor.SaveAsset"));
 					AssetSection.AddMenuEntryWithCommandList(FGenericCommands::Get().Paste, MenuContext->WeakCommandList.Pin());
 				}
 
 				/** Variable specific actions */
-				FToolMenuSection& VariablesSection = InMenu->AddSection("Variables", LOCTEXT("VariablesSectionLabel", "Variables"));			
+				FToolMenuSection& VariablesSection = InMenu->AddSection("Variables", LOCTEXT("VariablesSectionLabel", "Variables"));
 				{
+					// Populates entries for adding of different variable types/sources					
+					PopulateNewVariableToolMenuEntries(InMenu, true);
+					
 					VariablesSection.AddSubMenu("FindReferencesMenu", LOCTEXT("FindReferencesMenuLabel", "Find References..."), LOCTEXT("FindReferencesMenuToolTip", "Find references to selected variables by different means."),
 					FNewToolMenuDelegate::CreateLambda(
 						[MenuContext](UToolMenu* InSubmenu)
@@ -94,7 +114,7 @@ TSharedPtr<SWidget> FVariablesOutlinerMode::CreateContextMenu()
 								InSubmenu->AddMenuEntry(NAME_None, FToolMenuEntry::InitMenuEntryWithCommandList(FVariablesOutlinerCommands::Get().FindReferencesInAsset, MenuContext->WeakCommandList.Pin(), TAttribute<FText>(), TAttribute<FText>(),
 									UAnimNextAssetFindReplaceVariables::GetIconFromSearchScope(ESearchScope::Asset)));
 							}),
-						FToolUIActionChoice(),
+							FToolUIActionChoice(),
 							EUserInterfaceActionType::Button,
 							false,
 							FSlateIcon(FAppStyle::GetAppStyleSetName(), TEXT("Icons.Search"))
@@ -282,22 +302,6 @@ void FVariablesOutlinerMode::BindCommands(const TSharedRef<FUICommandList>& OutC
 		FIsActionButtonVisible::CreateRaw(this, &FVariablesOutlinerMode::CanDuplicate)
 	);
 
-	CommandList->MapAction(
-		FVariablesOutlinerCommands::Get().AddNewVariables,
-		FExecuteAction::CreateRaw(this, &FVariablesOutlinerMode::AddNewVariables),
-		FCanExecuteAction(),
-		FIsActionChecked(),
-		FIsActionButtonVisible::CreateRaw(this, &FVariablesOutlinerMode::CanAddVariable)
-		);
-
-	CommandList->MapAction(
-		FVariablesOutlinerCommands::Get().AddNewVariable,
-		FExecuteAction::CreateRaw(this, &FVariablesOutlinerMode::AddNewVariable),
-		FCanExecuteAction(),
-		FIsActionChecked(),
-		FIsActionButtonVisible::CreateRaw(this, &FVariablesOutlinerMode::CanAddVariable)
-		);
-
 	CommandList->MapAction(
 		FVariablesOutlinerCommands::Get().FindReferences,
 		FExecuteAction::CreateRaw(this, &FVariablesOutlinerMode::FindReferences, ESearchScope::Global),
@@ -623,162 +627,282 @@ bool FVariablesOutlinerMode::IsFindReferencesVisible(ESearchScope InSearchScope)
 	return CanFindReferences();
 }
 
-bool FVariablesOutlinerMode::CanAddVariable() const
+void FVariablesOutlinerMode::PopulateNewVariableToolMenuEntries(UToolMenu* InMenu, bool bAddSeparator)
 {
-	const FSceneOutlinerItemSelection& Selection = SceneOutliner->GetSelection();
-	if (Selection.Num() == 1)
+	const UAnimNextVariableItemMenuContext* MenuContext = InMenu->FindContext<UAnimNextVariableItemMenuContext>();
+	if(MenuContext == nullptr)
 	{
-		const FSceneOutlinerTreeItemPtr ItemToFind = Selection.SelectedItems[0].Pin();
-		return ItemToFind->IsA<FVariablesOutlinerAssetItem>() || ItemToFind->IsA<FVariablesOutlinerCategoryItem>();
+		return;
 	}
-	
-	return false;
-}
 
-void FVariablesOutlinerMode::AddNewVariable(UAnimNextRigVMAsset* InAsset)
-{
-	if (UAnimNextRigVMAssetEditorData* EditorData = UncookedOnly::FUtils::GetEditorData<UAnimNextRigVMAssetEditorData>(InAsset))
+	TMap<TWeakObjectPtr<UAnimNextRigVMAssetEditorData>, FString> AssetsAndCategoriesToAddTo;
+	if (MenuContext->WeakEditorDatas.Num())
 	{
-		UAnimNextVariableSettings* Settings = GetMutableDefault<UAnimNextVariableSettings>();
-		FScopedTransaction Transaction(LOCTEXT("AddVariable", "Add variable"));
-		
-		const FName VariableName = UncookedOnly::FUtils::GetValidVariableName(EditorData, Settings->GetLastVariableName());		
-		EditorData->AddVariable(VariableName, Settings->GetLastVariableType());
-		Settings->SetLastVariableName(VariableName);
-	}
-}
-
-void FVariablesOutlinerMode::SetHighlightedItem(FSceneOutlinerTreeItemID InID) const
-{
-	if (SVariablesOutliner* Outliner = static_cast<SVariablesOutliner*>(SceneOutliner))
-	{
-		if (FSceneOutlinerTreeItemPtr ItemPtr = Outliner->GetTreeItem(InID))
+		if (UAnimNextRigVMAssetEditorData* EditorData = MenuContext->WeakEditorDatas[0].Get())
 		{
-			Outliner->FrameItem(InID);
-			Outliner->SetHighlightedItem(ItemPtr);
+			AssetsAndCategoriesToAddTo.Add(EditorData, TEXT(""));
 		}
 	}
-}
-
-void FVariablesOutlinerMode::ClearHighlightedItem(FSceneOutlinerTreeItemID InID) const
-{
-	if (SVariablesOutliner* Outliner = static_cast<SVariablesOutliner*>(SceneOutliner))
+	else
 	{
-		if (FSceneOutlinerTreeItemPtr ItemPtr = Outliner->GetTreeItem(InID))
+		TSharedPtr<UE::AnimNext::Editor::SVariablesOutliner> SharedOutliner = MenuContext->WeakOutliner.Pin();
+		if (SharedOutliner.IsValid())
 		{
-			Outliner->ClearHighlightedItem(ItemPtr);
+			TArray<FSceneOutlinerTreeItemPtr> SelectedItems = SharedOutliner->GetSelectedItems();
+		
+			for(const FSceneOutlinerTreeItemPtr& Item : SelectedItems)
+			{
+				const FVariablesOutlinerAssetItem* AssetItem = Item->CastTo<FVariablesOutlinerAssetItem>();
+				const FVariablesOutlinerCategoryItem* CategoryItem = Item->CastTo<FVariablesOutlinerCategoryItem>();
+				if (AssetItem == nullptr && CategoryItem == nullptr)
+				{
+					continue;
+				}
+
+				UAnimNextRigVMAsset* Asset = nullptr;
+				FString Category;
+		
+				if (CategoryItem)
+				{
+					Asset = CategoryItem->WeakOwner.Get();
+					Category = CategoryItem->GetDisplayString();
+				}
+				else
+				{
+					Asset = AssetItem->SoftAsset.Get();
+				}
+
+				if(Asset == nullptr)
+				{
+					continue;
+				}
+
+				UAnimNextRigVMAssetEditorData* EditorData = UncookedOnly::FUtils::GetEditorData<UAnimNextRigVMAssetEditorData>(Asset);
+				if(EditorData == nullptr)
+				{
+					continue;
+				}
+
+				AssetsAndCategoriesToAddTo.Add(EditorData, Category);
+			}
 		}
 	}
-}
-
-void FVariablesOutlinerMode::AddNewVariable() const
-{
-	TArray<FSceneOutlinerTreeItemPtr> SelectedItems = SceneOutliner->GetSelectedItems();
 	
-	TMap<UAnimNextRigVMAssetEditorData*, FString> AssetsAndCategoriesToAddTo;
-	for(const FSceneOutlinerTreeItemPtr& Item : SelectedItems)
+	if (AssetsAndCategoriesToAddTo.Num())
 	{
-		const FVariablesOutlinerAssetItem* AssetItem = Item->CastTo<FVariablesOutlinerAssetItem>();
-		const FVariablesOutlinerCategoryItem* CategoryItem = Item->CastTo<FVariablesOutlinerCategoryItem>();
-		if (AssetItem == nullptr && CategoryItem == nullptr)
-		{			
-			continue;
-		}
+		constexpr bool bShouldCloseWindowAfterMenuSelection = true;
+		FMenuBuilder MenuBuilder(bShouldCloseWindowAfterMenuSelection, TSharedPtr<FUICommandList>());
 
-		UAnimNextRigVMAsset* Asset = nullptr;
-		FString Category;
-		
-		if (CategoryItem)
-		{
-			Asset = CategoryItem->WeakOwner.Get();
-			Category = CategoryItem->GetDisplayString();
-		}
-		else
-		{
-			Asset = AssetItem->SoftAsset.Get();
-		}
+		FToolMenuSection& Section = InMenu->AddSection("Variables", LOCTEXT("VariablesSectionLabel", "Variables"));
+		Section.AddMenuEntry("AddNewVariable", LOCTEXT("VariableLabel", "Add Variable"), LOCTEXT("VariableTooltip", "Adds a new Variable"),FSlateIcon(FAppStyle::GetAppStyleSetName(), "Icons.PlusCircle"),
+			FUIAction(FExecuteAction::CreateLambda([AssetsAndCategoriesToAddTo]()
+			{
+				for (auto& Pair : AssetsAndCategoriesToAddTo)
+				{
+					if (UAnimNextRigVMAssetEditorData* EditorData = Pair.Key.Get())
+                    {
+                        UAnimNextVariableSettings* Settings = GetMutableDefault<UAnimNextVariableSettings>();
+                        FScopedTransaction Transaction(LOCTEXT("AddVariable", "Add variable"));
+                        
+                        const FName VariableName = UncookedOnly::FUtils::GetValidVariableName(EditorData, Settings->GetLastVariableName());		
+
+						if (UAnimNextVariableEntry* VariableEntry = Pair.Key->AddVariable(VariableName, Settings->GetLastVariableType()))
+						{
+							VariableEntry->SetVariableCategory(Pair.Value);
+						}
+                        Settings->SetLastVariableName(VariableName);
+                    }
+				}
+			}))
+		);
 
-		if(Asset == nullptr)
+		if (AssetsAndCategoriesToAddTo.Num() == 1)
 		{
-			continue;
-		}
+			Section.AddSubMenu("AddSharedVariables", LOCTEXT("SharedVariableLabel", "Add Shared Variables"), LOCTEXT("SharedVariableTooltip", "Includes added SharedVariables"),
+			FNewMenuDelegate::CreateLambda([AssetsAndCategoriesToAddTo](FMenuBuilder& SubMenuBuilder)
+				{
+					TArray<FAssetData> ExistingSharedVariableReferences;
+					for (auto& Pair : AssetsAndCategoriesToAddTo)
+					{
+						if (UAnimNextRigVMAssetEditorData* EditorData = Pair.Key.Get())
+						{
+							EditorData->ForEachEntryOfType<UAnimNextSharedVariablesEntry>([&ExistingSharedVariableReferences](const UAnimNextSharedVariablesEntry* SharedVariablesEntry) -> bool
+							{
+								if (SharedVariablesEntry->GetType() == EAnimNextSharedVariablesType::Asset)
+								{
+									ExistingSharedVariableReferences.Add(SharedVariablesEntry->GetAsset());
+								}
+								return true;
+							});
+						}
+					}
 
-		UAnimNextRigVMAssetEditorData* EditorData = UncookedOnly::FUtils::GetEditorData<UAnimNextRigVMAssetEditorData>(Asset);
-		if(EditorData == nullptr)
-		{
-			continue;
+					FAssetPickerConfig AssetPickerConfig;
+					AssetPickerConfig.Filter.bRecursiveClasses = true;
+					AssetPickerConfig.Filter.ClassPaths.Add(UAnimNextSharedVariables::StaticClass()->GetClassPathName());
+					AssetPickerConfig.InitialAssetViewType = EAssetViewType::List;
+					AssetPickerConfig.OnAssetSelected = FOnAssetSelected::CreateLambda([AssetsAndCategoriesToAddTo](const FAssetData& InAssetData)
+					{
+						FSlateApplication::Get().DismissAllMenus();
+						if(UAnimNextSharedVariables* SharedVariables = Cast<UAnimNextSharedVariables>(InAssetData.GetAsset()))
+						{							
+							FScopedTransaction Transaction(LOCTEXT("AddSharedVariable", "Add Shared Variables"));
+							for (auto& Pair : AssetsAndCategoriesToAddTo)
+							{
+								if (UAnimNextRigVMAssetEditorData* EditorData = Pair.Key.Get())
+								{
+									EditorData->AddSharedVariables(SharedVariables);
+								}
+							}
+						}
+					});
+					AssetPickerConfig.OnShouldFilterAsset = FOnShouldFilterAsset::CreateLambda([ExistingSharedVariableReferences](const FAssetData& InAssetData)
+					{
+						FAnimNextAssetRegistryExports Exports;
+						UncookedOnly::FUtils::GetExportedVariablesForAsset(InAssetData, Exports);
+						if(Exports.Exports.Num() == 0 && InAssetData.GetClass(EResolveClass::Yes) != UAnimNextSharedVariables::StaticClass())
+						{
+							return true;
+						}
+
+						// Filter out already referenced assets
+						return ExistingSharedVariableReferences.ContainsByPredicate([InAssetData](const FAssetData& ExistingAssetData)
+						{
+							return ExistingAssetData == InAssetData;
+						});
+					});
+
+					FContentBrowserModule& ContentBrowserModule = FModuleManager::Get().LoadModuleChecked<FContentBrowserModule>(TEXT("ContentBrowser"));
+					TSharedPtr<SWidget> Widget = SNew(SBox)
+					.WidthOverride(300.0f)
+					.HeightOverride(400.0f)
+					[
+						ContentBrowserModule.Get().CreateAssetPicker(AssetPickerConfig)
+					];
+
+					SubMenuBuilder.AddWidget(Widget.ToSharedRef(), FText::GetEmpty(), true);
+				}),
+				false,
+				FSlateIcon(FAnimNextEditorStyle::Get().GetStyleSetName(), "ClassIcon.AnimNextSharedVariables")
+			);
+
+
+			Section.AddSubMenu("AddNativeSharedVariables",LOCTEXT("NativeSharedVariableLabel", "Add Native Shared Variables"), LOCTEXT("NativeSharedVariableTooltip", "Includes added Native SharedVariables"),
+		FNewMenuDelegate::CreateLambda([AssetsAndCategoriesToAddTo](FMenuBuilder& SubMenuBuilder)
+				{
+
+					TArray<FSoftObjectPath> ExistingSharedVariableStructPaths;
+					TArray<FAssetData> ExistingSharedVariableReferences;
+					for (auto& Pair : AssetsAndCategoriesToAddTo)
+					{
+						if (UAnimNextRigVMAssetEditorData* EditorData = Pair.Key.Get())
+						{
+							EditorData->ForEachEntryOfType<UAnimNextSharedVariablesEntry>([&ExistingSharedVariableStructPaths](const UAnimNextSharedVariablesEntry* SharedVariablesEntry) -> bool
+							{
+								if (SharedVariablesEntry->GetType() == EAnimNextSharedVariablesType::Struct)
+								{
+									ExistingSharedVariableStructPaths.Add(SharedVariablesEntry->GetObjectPath());
+								}
+								return true;
+							});
+						}
+					}
+			
+					class FStructFilter : public IStructViewerFilter
+					{
+					public:
+
+						FStructFilter(const TArray<FSoftObjectPath>& InFilteredStructPaths) : FilteredStructPaths(InFilteredStructPaths) {}
+						
+						virtual bool IsStructAllowed(const FStructViewerInitializationOptions& InInitOptions, const UScriptStruct* InStruct, TSharedRef<FStructViewerFilterFuncs> InFilterFuncs) override
+						{
+							if (InStruct->IsA<UUserDefinedStruct>())
+							{
+								return false;
+							}
+
+							if (!InStruct->HasMetaData(TEXT("BlueprintType")) || InStruct->HasMetaData(TEXT("Hidden")))
+							{
+								return false;
+							}
+
+							if (FilteredStructPaths.Contains(InStruct))
+							{
+								return false;
+							}
+
+							return true;
+						}
+
+						virtual bool IsUnloadedStructAllowed(const FStructViewerInitializationOptions& InInitOptions, const FSoftObjectPath& InStructPath, TSharedRef<class FStructViewerFilterFuncs> InFilterFuncs) override
+						{
+							return false;
+						};
+
+						TArray<FSoftObjectPath> FilteredStructPaths;
+					};
+
+					FStructViewerModule& StructViewerModule = FModuleManager::LoadModuleChecked<FStructViewerModule>("StructViewer");
+
+					FStructViewerInitializationOptions Options;
+					Options.Mode = EStructViewerMode::StructPicker;
+					Options.StructFilter = MakeShared<FStructFilter>(ExistingSharedVariableStructPaths);
+
+					TSharedPtr<SWidget> Widget = SNew(SBox)
+						.WidthOverride(300.0f)
+						.HeightOverride(400.0f)
+						[
+							StructViewerModule.CreateStructViewer(Options, FOnStructPicked::CreateLambda([AssetsAndCategoriesToAddTo](const UScriptStruct* InStruct)
+							{
+								FSlateApplication::Get().DismissAllMenus();
+
+								if(InStruct)
+								{
+									FScopedTransaction Transaction(LOCTEXT("AddNativeSharedVariable", "Add Native Shared Variables"));
+									for (auto& Pair : AssetsAndCategoriesToAddTo)
+									{
+										if (UAnimNextRigVMAssetEditorData* EditorData = Pair.Key.Get())
+										{									
+											EditorData->AddSharedVariablesStruct(InStruct);
+										}
+									}
+								}
+							}))
+						];
+
+					SubMenuBuilder.AddWidget(Widget.ToSharedRef(), FText::GetEmpty(), true);
+				}),
+				false,
+				FSlateIconFinder::FindIconForClass(UUserDefinedStruct::StaticClass())
+			);
 		}
 
-		AssetsAndCategoriesToAddTo.Add(EditorData, Category);
-	}
-	
-	UAnimNextVariableSettings* Settings = GetMutableDefault<UAnimNextVariableSettings>();
-	FScopedTransaction Transaction(LOCTEXT("AddVariable", "Add variable"));
-	for(const TPair<UAnimNextRigVMAssetEditorData*, FString>& Pair : AssetsAndCategoriesToAddTo)
-	{
-		const FName VariableName = UncookedOnly::FUtils::GetValidVariableName(Pair.Key, Settings->GetLastVariableName());
-		Settings->SetLastVariableName(VariableName);
-		
-		if (UAnimNextVariableEntry* VariableEntry = Pair.Key->AddVariable(VariableName, Settings->GetLastVariableType()))
+		if (bAddSeparator)
 		{
-			VariableEntry->SetVariableCategory(Pair.Value);
+			Section.AddSeparator("PostAddVariableSeparator");
 		}
 	}
 }
 
-void FVariablesOutlinerMode::AddNewVariables() const
+void FVariablesOutlinerMode::SetHighlightedItem(FSceneOutlinerTreeItemID InID) const
 {
-	TArray<FSceneOutlinerTreeItemPtr> SelectedItems = SceneOutliner->GetSelectedItems();
-	
-	TArray<UAnimNextRigVMAssetEditorData*> AssetsToAddTo;
-	for(const FSceneOutlinerTreeItemPtr& Item : SelectedItems)
+	if (SVariablesOutliner* Outliner = static_cast<SVariablesOutliner*>(SceneOutliner))
 	{
-		const FVariablesOutlinerAssetItem* AssetItem = Item->CastTo<FVariablesOutlinerAssetItem>();
-		if (AssetItem == nullptr)
-		{
-			continue;
-		}
-
-		UAnimNextRigVMAsset* Asset = AssetItem->SoftAsset.Get();
-		if(Asset == nullptr)
-		{
-			continue;
-		}
-
-		UAnimNextRigVMAssetEditorData* EditorData = UncookedOnly::FUtils::GetEditorData<UAnimNextRigVMAssetEditorData>(Asset);
-		if(EditorData == nullptr)
+		if (FSceneOutlinerTreeItemPtr ItemPtr = Outliner->GetTreeItem(InID))
 		{
-			continue;
+			Outliner->FrameItem(InID);
+			Outliner->SetHighlightedItem(ItemPtr);
 		}
-			
-		AssetsToAddTo.Add(EditorData);
 	}
-	
-	TSharedRef<SAddVariablesDialog> AddVariablesDialog = SNew(SAddVariablesDialog, AssetsToAddTo);
+}
 
-	TArray<FVariableToAdd> VariablesToAdd;
-	TArray<FSharedVariablesToAdd> AllSharedVariablesToAdd;
-	if(AddVariablesDialog->ShowModal(VariablesToAdd, AllSharedVariablesToAdd))
+void FVariablesOutlinerMode::ClearHighlightedItem(FSceneOutlinerTreeItemID InID) const
+{
+	if (SVariablesOutliner* Outliner = static_cast<SVariablesOutliner*>(SceneOutliner))
 	{
-		FScopedTransaction Transaction(FText::Format(LOCTEXT("AddVariablesFormat", "Add {0}|plural(one=variable, other=variables)"), (AllSharedVariablesToAdd.Num() + VariablesToAdd.Num()) * AssetsToAddTo.Num()));
-		for(UAnimNextRigVMAssetEditorData* EditorData : AssetsToAddTo)
+		if (FSceneOutlinerTreeItemPtr ItemPtr = Outliner->GetTreeItem(InID))
 		{
-			for (const FVariableToAdd& VariableToAdd : VariablesToAdd)
-			{
-				EditorData->AddVariable(VariableToAdd.Name, VariableToAdd.Type);
-			}
-
-			for (const FSharedVariablesToAdd& SharedVariablesToAdd : AllSharedVariablesToAdd)
-			{
-				if (SharedVariablesToAdd.SharedVariables.IsType<UAnimNextSharedVariables*>())
-				{
-					EditorData->AddSharedVariables(SharedVariablesToAdd.SharedVariables.Get<UAnimNextSharedVariables*>());
-				}
-				else
-				{
-					EditorData->AddSharedVariablesStruct(SharedVariablesToAdd.SharedVariables.Get<const UScriptStruct*>());
-				}
-			}
+			Outliner->ClearHighlightedItem(ItemPtr);
 		}
 	}
 }
diff --git a/Engine/Plugins/Experimental/AnimNext/Source/AnimNextEditor/Private/Variables/Outliner/VariablesOutlinerMode.h b/Engine/Plugins/Experimental/AnimNext/Source/AnimNextEditor/Private/Variables/Outliner/VariablesOutlinerMode.h
--- a/Engine/Plugins/Experimental/AnimNext/Source/AnimNextEditor/Private/Variables/Outliner/VariablesOutlinerMode.h
+++ b/Engine/Plugins/Experimental/AnimNext/Source/AnimNextEditor/Private/Variables/Outliner/VariablesOutlinerMode.h
@@ -45,10 +45,10 @@ public:
 	virtual bool CanSupportDragAndDrop() const override { return true; }
 	virtual void OnItemDoubleClick(FSceneOutlinerTreeItemPtr Item) override;
 	
-	static void AddNewVariable(UAnimNextRigVMAsset* InAsset);
-
 	void SetHighlightedItem(FSceneOutlinerTreeItemID InID) const;
 	void ClearHighlightedItem(FSceneOutlinerTreeItemID InID) const;
+
+	static void PopulateNewVariableToolMenuEntries(UToolMenu* InMenu, bool bAddSeparator);
 protected:
 	virtual TUniquePtr<ISceneOutlinerHierarchy> CreateHierarchy() override;
 	// End ISceneOutlinerMode overrides
@@ -79,10 +79,6 @@ protected:
 	bool CanFindReferences() const;
 	bool IsFindReferencesVisible(ESearchScope InSearchScope) const;
 
-	bool CanAddVariable() const;
-	void AddNewVariable() const;
-	void AddNewVariables() const;
-
 	void SaveAsset() const;
 	bool CanSaveAsset() const;
 	bool IsAsset() const;
diff --git a/Engine/Plugins/Experimental/AnimNext/Source/AnimNextEditor/Private/Variables/SAddVariablesDialog.cpp b/Engine/Plugins/Experimental/AnimNext/Source/AnimNextEditor/Private/Variables/SAddVariablesDialog.cpp
deleted file mode 100644
diff --git a/Engine/Plugins/Experimental/AnimNext/Source/AnimNextEditor/Private/Variables/SAddVariablesDialog.h b/Engine/Plugins/Experimental/AnimNext/Source/AnimNextEditor/Private/Variables/SAddVariablesDialog.h
deleted file mode 100644
diff --git a/Engine/Plugins/Experimental/AnimNext/Source/AnimNextEditor/Private/Variables/SVariablePicker.cpp b/Engine/Plugins/Experimental/AnimNext/Source/AnimNextEditor/Private/Variables/SVariablePicker.cpp
--- a/Engine/Plugins/Experimental/AnimNext/Source/AnimNextEditor/Private/Variables/SVariablePicker.cpp
+++ b/Engine/Plugins/Experimental/AnimNext/Source/AnimNextEditor/Private/Variables/SVariablePicker.cpp
@@ -4,7 +4,6 @@
 
 #include "UncookedOnlyUtils.h"
 #include "Param/ParamType.h"
-#include "SAddVariablesDialog.h"
 #include "Widgets/Input/SSearchBox.h"
 #include "Widgets/Layout/SBox.h"
 #include "Widgets/Images/SImage.h"
diff --git a/Engine/Plugins/Experimental/AnimNext/Source/AnimNextEditor/Private/Variables/SVariablesView.cpp b/Engine/Plugins/Experimental/AnimNext/Source/AnimNextEditor/Private/Variables/SVariablesView.cpp
--- a/Engine/Plugins/Experimental/AnimNext/Source/AnimNextEditor/Private/Variables/SVariablesView.cpp
+++ b/Engine/Plugins/Experimental/AnimNext/Source/AnimNextEditor/Private/Variables/SVariablesView.cpp
@@ -10,12 +10,10 @@
 #include "SceneOutlinerPublicTypes.h"
 #include "SceneOutlinerSourceControlColumn.h"
 #include "ScopedTransaction.h"
-#include "SPositiveActionButton.h"
 #include "SSceneOutliner.h"
 #include "Outliner/VariablesOutlinerColumns.h"
 #include "Outliner/VariablesOutlinerMode.h"
 #include "Outliner/VariablesOutlinerEntryItem.h"
-#include "Variables/SAddVariablesDialog.h"
 #include "IWorkspaceEditor.h"
 #include "IWorkspaceEditorModule.h"
 #include "WorkspaceAssetRegistryInfo.h"
@@ -29,133 +27,6 @@ namespace UE::AnimNext::Editor
 
 const FLazyName VariablesTabName("VariablesTab");
 
-void SVariablesOutliner::Construct(const FArguments& InArgs, const FSceneOutlinerInitializationOptions& InitOptions)
-{
-	AddVariablesButton = SNew(SPositiveActionButton)
-			.OnClicked(this, &SVariablesOutliner::HandleAddVariablesClicked)
-			.Icon(FAppStyle::Get().GetBrush("Icons.Plus"))
-			.Text(LOCTEXT("AddVariablesButton", "Add"))
-			.ToolTipText(LOCTEXT("AddVariablesButtonTooltip", "Adds variables to assets.\nIf multiple assets are selected, then variables will be added to each.\nIf no assets are selected and there are multiple assets, variables will be added to all assets."));
-
-	SSceneOutliner::Construct(InArgs, InitOptions);
-}
-
-void SVariablesOutliner::CustomAddToToolbar(TSharedPtr<class SHorizontalBox> Toolbar)
-{
-	Toolbar->AddSlot()
-	.VAlign(VAlign_Center)
-	.AutoWidth()
-	.Padding(4.f, 0.f, 0.f, 0.f)
-	[
-		AddVariablesButton.ToSharedRef()
-	];
-}
-
-FReply SVariablesOutliner::HandleAddVariablesClicked()
-{
-	TArray<UAnimNextRigVMAssetEditorData*> AssetsToAddTo;
-	if(Assets.Num() == 1)
-	{
-		UAnimNextRigVMAsset* Asset = Assets[0].Get();
-		if(Asset == nullptr)
-		{
-			FReply::Unhandled();
-		}
-
-		UAnimNextRigVMAssetEditorData* EditorData = UncookedOnly::FUtils::GetEditorData<UAnimNextRigVMAssetEditorData>(Asset);
-		if(EditorData == nullptr)
-		{
-			FReply::Unhandled();
-		}
-
-		AssetsToAddTo.Add(EditorData);
-	}
-	else
-	{
-		// Add to multiple selected
-		TArray<FSceneOutlinerTreeItemPtr> SelectedItems = GetSelectedItems();
-		for(const FSceneOutlinerTreeItemPtr& Item : SelectedItems)
-		{
-			const FVariablesOutlinerAssetItem* AssetItem = Item->CastTo<FVariablesOutlinerAssetItem>();
-			if (AssetItem == nullptr)
-			{
-				continue;
-			}
-
-			UAnimNextRigVMAsset* Asset = AssetItem->SoftAsset.Get();
-			if(Asset == nullptr)
-			{
-				continue;
-			}
-
-			UAnimNextRigVMAssetEditorData* EditorData = UncookedOnly::FUtils::GetEditorData<UAnimNextRigVMAssetEditorData>(Asset);
-			if(EditorData == nullptr)
-			{
-				continue;
-			}
-			
-			AssetsToAddTo.Add(EditorData);
-		}
-
-		// No selected asset items, so add to all asset items
-		if(AssetsToAddTo.Num() == 0)
-		{
-			for(const TSoftObjectPtr<UAnimNextRigVMAsset>& SoftAsset : Assets)
-			{
-				UAnimNextRigVMAsset* Asset = SoftAsset.Get();
-				if(Asset == nullptr)
-				{
-					continue;
-				}
-
-				UAnimNextRigVMAssetEditorData* EditorData = UncookedOnly::FUtils::GetEditorData<UAnimNextRigVMAssetEditorData>(Asset);
-				if(EditorData == nullptr)
-				{
-					continue;
-				}
-
-				AssetsToAddTo.Add(EditorData);
-			}
-		}
-	}
-
-	if(AssetsToAddTo.Num() == 0)
-	{
-		return FReply::Unhandled();
-	}
-
-	TSharedRef<SAddVariablesDialog> AddVariablesDialog =
-		SNew(SAddVariablesDialog, AssetsToAddTo);
-
-	TArray<FVariableToAdd> VariablesToAdd;
-	TArray<FSharedVariablesToAdd> AllSharedVariablesToAdd;
-	if(AddVariablesDialog->ShowModal(VariablesToAdd, AllSharedVariablesToAdd))
-	{
-		FScopedTransaction Transaction(FText::Format(LOCTEXT("AddVariablesFormat", "Add {0}|plural(one=variable, other=variables)"), (AllSharedVariablesToAdd.Num() + VariablesToAdd.Num()) * AssetsToAddTo.Num()));
-		for(UAnimNextRigVMAssetEditorData* EditorData : AssetsToAddTo)
-		{
-			for (const FVariableToAdd& VariableToAdd : VariablesToAdd)
-			{
-				EditorData->AddVariable(VariableToAdd.Name, VariableToAdd.Type);
-			}
-
-			for (const FSharedVariablesToAdd& SharedVariablesToAdd : AllSharedVariablesToAdd)
-			{
-				if (SharedVariablesToAdd.SharedVariables.IsType<UAnimNextSharedVariables*>())
-				{
-					EditorData->AddSharedVariables(SharedVariablesToAdd.SharedVariables.Get<UAnimNextSharedVariables*>());
-				}
-				else
-				{
-					EditorData->AddSharedVariablesStruct(SharedVariablesToAdd.SharedVariables.Get<const UScriptStruct*>());
-				}
-			}
-		}
-	}
-
-	return FReply::Handled();
-}
-
 void SVariablesOutliner::RegisterAssetDelegates(const UAnimNextRigVMAsset* InAsset)
 {	
 	// Bind for any modification callbacks
diff --git a/Engine/Plugins/Experimental/AnimNext/Source/AnimNextEditor/Private/Variables/SVariablesView.h b/Engine/Plugins/Experimental/AnimNext/Source/AnimNextEditor/Private/Variables/SVariablesView.h
--- a/Engine/Plugins/Experimental/AnimNext/Source/AnimNextEditor/Private/Variables/SVariablesView.h
+++ b/Engine/Plugins/Experimental/AnimNext/Source/AnimNextEditor/Private/Variables/SVariablesView.h
@@ -34,14 +34,7 @@ extern const FLazyName VariablesTabName;
 class SVariablesOutliner : public SSceneOutliner
 {
 public:
-	void Construct(const FArguments& InArgs, const FSceneOutlinerInitializationOptions& InitOptions);
-
-	virtual void CustomAddToToolbar(TSharedPtr<SHorizontalBox> Toolbar) override;
-
-	FReply HandleAddVariablesClicked();
-
 	void SetExport(const FWorkspaceOutlinerItemExport& InExport);
-
 	void OnEditorDataModified(UAnimNextRigVMAssetEditorData* InEditorData, EAnimNextEditorDataNotifType InType, UObject* InSubject);
 
 	void SetHighlightedItem(FSceneOutlinerTreeItemPtr Item) const;
@@ -61,8 +54,6 @@ private:
 	friend class FVariablesOutlinerHierarchy;
 
 	TArray<TSoftObjectPtr<UAnimNextRigVMAsset>> Assets;
-	TSharedPtr<SPositiveActionButton> AddVariablesButton;
-
 	FWorkspaceOutlinerItemExport Export;
 };
 
