Hide Find References Submenu entry when selection does not allow for it

diff --git a/Engine/Plugins/Experimental/AnimNext/Source/AnimNextEditor/Private/Variables/Outliner/VariablesOutlinerMode.cpp b/Engine/Plugins/Experimental/AnimNext/Source/AnimNextEditor/Private/Variables/Outliner/VariablesOutlinerMode.cpp
--- a/Engine/Plugins/Experimental/AnimNext/Source/AnimNextEditor/Private/Variables/Outliner/VariablesOutlinerMode.cpp
+++ b/Engine/Plugins/Experimental/AnimNext/Source/AnimNextEditor/Private/Variables/Outliner/VariablesOutlinerMode.cpp
@@ -114,7 +114,18 @@ TSharedPtr<SWidget> FVariablesOutlinerMode::CreateContextMenu()
 								InSubmenu->AddMenuEntry(NAME_None, FToolMenuEntry::InitMenuEntryWithCommandList(FVariablesOutlinerCommands::Get().FindReferencesInAsset, MenuContext->WeakCommandList.Pin(), TAttribute<FText>(), TAttribute<FText>(),
 									UAnimNextAssetFindReplaceVariables::GetIconFromSearchScope(ESearchScope::Asset)));
 							}),
-							FToolUIActionChoice(),
+							FUIAction(FExecuteAction(), FCanExecuteAction(), FGetActionCheckState(), FIsActionButtonVisible::CreateLambda([MenuContext]() -> bool
+							{
+								if (TSharedPtr<UE::AnimNext::Editor::SVariablesOutliner> SharedOutliner = MenuContext->WeakOutliner.Pin())
+								{
+									if (const FVariablesOutlinerMode* Mode = static_cast<const FVariablesOutlinerMode*>(SharedOutliner->GetMode()))
+									{
+										return Mode->CanFindReferences();
+									}
+								}
+
+								return false;
+							})),
 							EUserInterfaceActionType::Button,
 							false,
 							FSlateIcon(FAppStyle::GetAppStyleSetName(), TEXT("Icons.Search"))
