Added DragDrop type for Workspace editor - holds FAssetData for all top-level selected (dragged) assets
Added support for Workspace DragDrop to implement/add SharedVariabl asset in VariablesOutliner

diff --git a/Engine/Plugins/Experimental/AnimNext/Source/AnimNextEditor/Private/Variables/Outliner/VariablesOutlinerAssetItem.cpp b/Engine/Plugins/Experimental/AnimNext/Source/AnimNextEditor/Private/Variables/Outliner/VariablesOutlinerAssetItem.cpp
--- a/Engine/Plugins/Experimental/AnimNext/Source/AnimNextEditor/Private/Variables/Outliner/VariablesOutlinerAssetItem.cpp
+++ b/Engine/Plugins/Experimental/AnimNext/Source/AnimNextEditor/Private/Variables/Outliner/VariablesOutlinerAssetItem.cpp
@@ -17,6 +17,7 @@
 #include "UncookedOnlyUtils.h"
 #include "VariablesOutlinerDragDrop.h"
 #include "VariablesOutlinerMode.h"
+#include "WorkspaceDragDropOperation.h"
 #include "Styling/SlateIconFinder.h"
 #include "Styling/StarshipCoreStyle.h"
 #include "UObject/Package.h"
@@ -377,9 +378,49 @@ class SVariablesOutlinerAssetLabel : FSceneOutlinerCommonLabelData, public SComp
 							const FText FormattedMessage = FText::Format(LOCTEXT("MoveVariableToAssetFormat", "Moving {0} to {1}"), VariableNameText, AssetNameText);
 							FScopedTransaction Transaction(FormattedMessage);
 							UncookedOnly::FUtils::MoveVariableToAsset(VariableEntry, Asset, true, true);
+
+							return FReply::Handled();
+						}
+					}
+				}
+			}
+		}
+
+		TSharedPtr<FWorkspaceDragDropOp> WorkspaceDragDrop  = DragDropEvent.GetOperationAs<FWorkspaceDragDropOp>();
+		if (WorkspaceDragDrop.IsValid())
+		{
+			if (UAnimNextRigVMAsset* Asset = TreeItem.Pin()->SoftAsset.Get())
+			{
+				const FText AssetNameText = FText::FromName(Asset->GetFName());
+
+				UAnimNextRigVMAssetEditorData* EditorData = UncookedOnly::FUtils::GetEditorData<UAnimNextRigVMAssetEditorData>(Asset);
+				TArray<const UAnimNextSharedVariables*> SharedVariableAssetsToAdd;				
+				for (const FAssetData& AssetData : WorkspaceDragDrop->AssetDatas)
+				{					
+					const UClass* AssetClass = AssetData.GetClass(EResolveClass::Yes);
+					if (AssetClass && AssetClass->IsChildOf<UAnimNextSharedVariables>())
+					{						
+						if (!EditorData->FindEntry(AssetData.AssetName))
+						{
+							if (const UAnimNextSharedVariables* SharedVariables = Cast<UAnimNextSharedVariables>(AssetData.GetAsset()))
+							{
+								SharedVariableAssetsToAdd.Add(SharedVariables);
+							}
 						}
 					}
 				}
+
+				if (SharedVariableAssetsToAdd.Num())
+				{
+					const FText FormattedMessage = FText::Format(LOCTEXT("AddSharedVariablesToAssetTransaction", "Adding Shared Variables to {0}"), AssetNameText);
+					FScopedTransaction Transaction(FormattedMessage);
+					for (const UAnimNextSharedVariables* SharedVariables : SharedVariableAssetsToAdd)
+					{
+						EditorData->AddSharedVariables(SharedVariables);
+					}
+
+					return FReply::Handled();
+				}
 			}
 		}
 		
@@ -418,6 +459,77 @@ class SVariablesOutlinerAssetLabel : FSceneOutlinerCommonLabelData, public SComp
 				}
 			}
 		}
+
+		TSharedPtr<FWorkspaceDragDropOp> WorkspaceDragDrop  = DragDropEvent.GetOperationAs<FWorkspaceDragDropOp>();
+		if (WorkspaceDragDrop.IsValid())
+		{
+			if (const UAnimNextRigVMAsset* Asset = TreeItem.Pin()->SoftAsset.Get())
+			{
+				const FText AssetNameText = FText::FromName(Asset->GetFName());
+
+				TArray<FName> SharedVariableAssetNamesToAdd;
+				TArray<FName> PreexistingSharedVariableAssetNames;
+				
+				for (const FAssetData& AssetData : WorkspaceDragDrop->AssetDatas)
+				{					
+					const UClass* AssetClass = AssetData.GetClass(EResolveClass::Yes);
+					if (AssetClass && AssetClass->IsChildOf<UAnimNextSharedVariables>())
+					{						
+						const UAnimNextRigVMAssetEditorData* EditorData = UncookedOnly::FUtils::GetEditorData<const UAnimNextRigVMAssetEditorData>(Asset);
+						if (!EditorData->FindEntry(AssetData.AssetName))
+						{
+							SharedVariableAssetNamesToAdd.Add(AssetData.AssetName);
+						}
+						else
+						{
+
+							PreexistingSharedVariableAssetNames.Add(AssetData.AssetName);
+						}
+					}
+				}
+
+				const bool bAddingSharedVariables = SharedVariableAssetNamesToAdd.Num() > 0;
+				const bool bPreexistingSharedVariables = PreexistingSharedVariableAssetNames.Num() > 0;
+
+				if (bPreexistingSharedVariables && !bAddingSharedVariables)
+				{
+					const bool bMultipleAssets = PreexistingSharedVariableAssetNames.Num() > 1;
+					
+					const FSlateBrush* ErrorSymbol = FAppStyle::GetBrush(TEXT("Graph.ConnectorFeedback.Error"));
+					const FText SingleFormattedMessage = FText::Format(LOCTEXT("ExistingSharedVariablesInAsset", "{0} already references {1}"), AssetNameText, FText::FromName(PreexistingSharedVariableAssetNames[0]));
+					const FText MultiFormattedMessage = FText::Format(LOCTEXT("ExistingSharedVariablesInAsset", "{0} already references assets"), AssetNameText);
+
+					WorkspaceDragDrop->SetToolTip(bMultipleAssets ? MultiFormattedMessage : SingleFormattedMessage, ErrorSymbol);
+				}
+				else if (bAddingSharedVariables)
+				{
+					if (SharedVariableAssetNamesToAdd.Num() > 1)
+					{
+						FString SharedVariablesNameListString;
+						for (const FName& AssetName : SharedVariableAssetNamesToAdd)
+						{
+							if (!SharedVariablesNameListString.IsEmpty())
+							{
+								SharedVariablesNameListString.Append(TEXT(", "));
+							}
+
+							SharedVariablesNameListString.Appendf(TEXT("%s"), *AssetName.ToString());
+						}
+						
+						const FText SharedVariablesNameListText = FText::FromString(SharedVariablesNameListString);
+						const FText MultiFormattedMessage = FText::Format(LOCTEXT("AddMultipleSharedVariablesToAsset", "Add {0} to {1}"), SharedVariablesNameListText, AssetNameText);
+
+						WorkspaceDragDrop->SetToolTip(MultiFormattedMessage, FSlateIconFinder::FindIconBrushForClass(UAnimNextSharedVariables::StaticClass()));
+					}
+					else
+					{
+						const FText SharedVariablesNameText = FText::FromName(SharedVariableAssetNamesToAdd[0]);
+						const FText SingleFormattedMessage = FText::Format(LOCTEXT("AddSharedVariablesToAsset", "Add {0} to {1}"), SharedVariablesNameText, AssetNameText);
+						WorkspaceDragDrop->SetToolTip(SingleFormattedMessage, FSlateIconFinder::FindIconBrushForClass(UAnimNextSharedVariables::StaticClass()));
+					}				
+				}
+			}
+		}
 	}
 
 	virtual void OnDragLeave( const FDragDropEvent& DragDropEvent ) override
diff --git a/Engine/Plugins/Experimental/Workspace/Source/WorkspaceEditor/Private/Outliner/WorkspaceOutlinerMode.cpp b/Engine/Plugins/Experimental/Workspace/Source/WorkspaceEditor/Private/Outliner/WorkspaceOutlinerMode.cpp
--- a/Engine/Plugins/Experimental/Workspace/Source/WorkspaceEditor/Private/Outliner/WorkspaceOutlinerMode.cpp
+++ b/Engine/Plugins/Experimental/Workspace/Source/WorkspaceEditor/Private/Outliner/WorkspaceOutlinerMode.cpp
@@ -14,6 +14,7 @@
 #include "AssetRegistry/AssetRegistryModule.h"
 #include "ScopedTransaction.h"
 #include "ToolMenus.h"
+#include "WorkspaceDragDropOperation.h"
 #include "WorkspaceEditorCommands.h"
 #include "Framework/Commands/GenericCommands.h"
 #include "Algo/ForEach.h"
@@ -376,6 +377,51 @@ void FWorkspaceOutlinerMode::BindCommands(const TSharedRef<FUICommandList>& OutC
 	);
 }
 
+TSharedPtr<FDragDropOperation> FWorkspaceOutlinerMode::CreateDragDropOperation(const FPointerEvent& MouseEvent,
+	const TArray<FSceneOutlinerTreeItemPtr>& InTreeItems) const
+{
+ 	if(InTreeItems.Num() == 0)
+	{
+		return nullptr;
+	}
+
+	TArray<FSoftObjectPath> AssetPaths;
+	for (const FSceneOutlinerTreeItemPtr& Item : InTreeItems)
+	{
+		if (const FWorkspaceOutlinerTreeItem* WorkspaceItem = Item->CastTo<FWorkspaceOutlinerTreeItem>())
+		{
+			// Root level entries only (asset) 
+			if (WorkspaceItem->Export.GetParentIdentifier() == NAME_None)
+			{
+				const FSoftObjectPath AssetPath = WorkspaceItem->Export.GetFirstAssetPath();
+				if (AssetPath.IsValid())
+				{
+					AssetPaths.AddUnique(AssetPath);
+				}
+			}
+		}
+	}
+
+	if (AssetPaths.Num())
+	{
+		const IAssetRegistry& AssetRegistry = FAssetRegistryModule::GetRegistry();
+		FARFilter Filter;	
+		Filter.SoftObjectPaths = MoveTemp(AssetPaths);
+		
+		TArray<FAssetData> AssetDataList;
+		AssetRegistry.GetAssets(Filter, AssetDataList);
+
+
+		TSharedPtr<FWorkspaceDragDropOp> DragDropOp = MakeShared<FWorkspaceDragDropOp>();
+		DragDropOp->Init(AssetDataList);
+		DragDropOp->Construct();
+		
+		return DragDropOp;
+	}
+ 
+	return nullptr;
+}
+
 TUniquePtr<ISceneOutlinerHierarchy> FWorkspaceOutlinerMode::CreateHierarchy()
 {
 	return MakeUnique<FWorkspaceOutlinerHierarchy>(this, WeakWorkspace);
diff --git a/Engine/Plugins/Experimental/Workspace/Source/WorkspaceEditor/Private/Outliner/WorkspaceOutlinerMode.h b/Engine/Plugins/Experimental/Workspace/Source/WorkspaceEditor/Private/Outliner/WorkspaceOutlinerMode.h
--- a/Engine/Plugins/Experimental/Workspace/Source/WorkspaceEditor/Private/Outliner/WorkspaceOutlinerMode.h
+++ b/Engine/Plugins/Experimental/Workspace/Source/WorkspaceEditor/Private/Outliner/WorkspaceOutlinerMode.h
@@ -28,6 +28,7 @@ namespace UE::Workspace
 		virtual bool CanRename() const override;
 		virtual bool CanRenameItem(const ISceneOutlinerTreeItem& Item) const override;
 		virtual void BindCommands(const TSharedRef<FUICommandList>& OutCommandList) override;
+		virtual TSharedPtr<FDragDropOperation> CreateDragDropOperation(const FPointerEvent& MouseEvent, const TArray<FSceneOutlinerTreeItemPtr>& InTreeItems) const override;
 	protected:
 		virtual TUniquePtr<ISceneOutlinerHierarchy> CreateHierarchy() override;		
 		// End ISceneOutlinerMode overrides
diff --git a/Engine/Plugins/Experimental/Workspace/Source/WorkspaceEditor/Public/WorkspaceDragDropOperation.h b/Engine/Plugins/Experimental/Workspace/Source/WorkspaceEditor/Public/WorkspaceDragDropOperation.h
new file mode 100644
--- /dev/null
+++ b/Engine/Plugins/Experimental/Workspace/Source/WorkspaceEditor/Public/WorkspaceDragDropOperation.h
@@ -0,0 +1,39 @@
+// Copyright Epic Games, Inc. All Rights Reserved.
+
+#pragma once
+
+#include "CoreMinimal.h"
+#include "ClassIconFinder.h"
+#include "Input/DragAndDrop.h"
+#include "DragAndDrop/DecoratedDragDropOp.h"
+#include "AssetRegistry/AssetData.h"
+
+class FWorkspaceDragDropOp : public FDecoratedDragDropOp
+{
+public:
+	DRAG_DROP_OPERATOR_TYPE(FWorkspaceDragDropOp, FDecoratedDragDropOp)
+
+	/** Root-level Asset data workspace entries */
+	TArray<FAssetData> AssetDatas;
+
+	void Init(const TArray<FAssetData>& InAssetDatas) 
+	{
+		AssetDatas = InAssetDatas;
+
+		if(AssetDatas.Num() == 0)
+		{
+			CurrentHoverText = NSLOCTEXT("FWorkspaceDragDropOp", "None", "None");
+		}
+		else 
+		{
+			const FText FirstAssetText = FText::FromName(AssetDatas[0].AssetName);
+			const FText ExtraAssetsText = FText::Format(NSLOCTEXT("FWorkspaceDragDropOp", "TooltipAssetsFormat", " and {0} {1}|plural(one=other, other=others)"), FText::AsNumber(AssetDatas.Num() - 1), AssetDatas.Num() - 1);
+						
+			CurrentHoverText = FText::Format(NSLOCTEXT("FWorkspaceDragDropOp", "TooltipFormat", "{0}{1}"), FirstAssetText, AssetDatas.Num() > 1 ? ExtraAssetsText : FText::GetEmpty() );
+		}
+
+		CurrentIconBrush = FClassIconFinder::FindThumbnailForClass(UObject::StaticClass());
+
+		SetupDefaults();
+	}
+};
