Fix crashes after undo'ing add function to any anim graph.
- Crash when deleting all functions after any add undo.
- Crash when saving asset after an undo.

diff --git a/Engine/Plugins/Experimental/AnimNext/Source/AnimNextUncookedOnly/Private/AnimNextRigVMAssetEditorData.cpp b/Engine/Plugins/Experimental/AnimNext/Source/AnimNextUncookedOnly/Private/AnimNextRigVMAssetEditorData.cpp
--- a/Engine/Plugins/Experimental/AnimNext/Source/AnimNextUncookedOnly/Private/AnimNextRigVMAssetEditorData.cpp
+++ b/Engine/Plugins/Experimental/AnimNext/Source/AnimNextUncookedOnly/Private/AnimNextRigVMAssetEditorData.cpp
@@ -2452,6 +2452,11 @@ bool UAnimNextRigVMAssetLibrary::RenameCategory(UAnimNextRigVMAsset* InAsset, co
 
 URigVMLibraryNode* UAnimNextRigVMAssetEditorData::AddFunction(FName InFunctionName, bool bInMutable, bool bInSetupUndoRedo, bool bPrintPythonCommand)
 {
+	if (bInSetupUndoRedo)
+	{
+		Modify();
+	}
+
 	URigVMController* Controller = RigVMClient.GetOrCreateController(GetLocalFunctionLibrary());
 	URigVMLibraryNode* Node = Controller->AddFunctionToLibrary(InFunctionName, bInMutable, FVector2D::ZeroVector, bInSetupUndoRedo, false);
 
diff --git a/Engine/Plugins/Runtime/RigVM/Source/RigVMDeveloper/Private/RigVMModel/RigVMController.cpp b/Engine/Plugins/Runtime/RigVM/Source/RigVMDeveloper/Private/RigVMModel/RigVMController.cpp
--- a/Engine/Plugins/Runtime/RigVM/Source/RigVMDeveloper/Private/RigVMModel/RigVMController.cpp
+++ b/Engine/Plugins/Runtime/RigVM/Source/RigVMDeveloper/Private/RigVMModel/RigVMController.cpp
@@ -14513,6 +14513,11 @@ URigVMLibraryNode* URigVMController::AddFunctionToLibrary(const FName& InFunctio
 
 	if (URigVMFunctionLibrary* Library = Cast<URigVMFunctionLibrary>(Graph))
 	{
+		if (bSetupUndoRedo)
+		{
+			Library->Modify();
+		}
+
 		FRigVMVariant Variant;
 		if (FRigVMGraphFunctionStore* FunctionStore = GetGraphFunctionStore())
 		{
