UAF: Namespace renamed

diff --git a/Engine/Plugins/Experimental/AnimNext/Source/AnimNext/Internal/AnimNextDebugDraw.h b/Engine/Plugins/Experimental/AnimNext/Source/AnimNext/Internal/AnimNextDebugDraw.h
--- a/Engine/Plugins/Experimental/AnimNext/Source/AnimNext/Internal/AnimNextDebugDraw.h
+++ b/Engine/Plugins/Experimental/AnimNext/Source/AnimNext/Internal/AnimNextDebugDraw.h
@@ -13,12 +13,12 @@ struct FRigVMDrawInterface;
 struct FAnimNextModuleInstance;
 class UAnimNextWorldSubsystem;
 
-namespace UE::AnimNext
+namespace UE::UAF
 {
 	struct FModuleEventTickFunction;
 }
 
-namespace UE::AnimNext::Debug
+namespace UE::UAF::Debug
 {
 
 #if UE_ENABLE_DEBUG_DRAWING
@@ -89,7 +89,7 @@ private:
 	FRWLock Lock;
 
 	friend struct ::FAnimNextModuleInstance;
-	friend struct UE::AnimNext::FModuleEventTickFunction;
+	friend struct UE::UAF::FModuleEventTickFunction;
 	friend class ::UAnimNextWorldSubsystem;
 };
 
diff --git a/Engine/Plugins/Experimental/AnimNext/Source/AnimNext/Internal/AnimNextExecuteContext.h b/Engine/Plugins/Experimental/AnimNext/Source/AnimNext/Internal/AnimNextExecuteContext.h
--- a/Engine/Plugins/Experimental/AnimNext/Source/AnimNext/Internal/AnimNextExecuteContext.h
+++ b/Engine/Plugins/Experimental/AnimNext/Source/AnimNext/Internal/AnimNextExecuteContext.h
@@ -9,7 +9,7 @@
 
 struct FUAFAssetInstance;
 
-namespace UE::AnimNext
+namespace UE::UAF
 {
 	struct FLatentPropertyHandle;
 	struct FModuleEventTickFunction;
@@ -48,10 +48,10 @@ protected:
 	// Context data for this execution
 	TStructView<FAnimNextModuleContextData> ContextData;
 
-	friend struct UE::AnimNext::FScopedExecuteContextData;
+	friend struct UE::UAF::FScopedExecuteContextData;
 };
 
-namespace UE::AnimNext
+namespace UE::UAF
 {
 
 // Helper for applying context data prior to RigVM execution
diff --git a/Engine/Plugins/Experimental/AnimNext/Source/AnimNext/Internal/AnimNextModuleImpl.h b/Engine/Plugins/Experimental/AnimNext/Source/AnimNext/Internal/AnimNextModuleImpl.h
--- a/Engine/Plugins/Experimental/AnimNext/Source/AnimNext/Internal/AnimNextModuleImpl.h
+++ b/Engine/Plugins/Experimental/AnimNext/Source/AnimNext/Internal/AnimNextModuleImpl.h
@@ -5,7 +5,7 @@
 #include "CoreMinimal.h"
 #include "Modules/ModuleInterface.h"
 
-namespace UE::AnimNext
+namespace UE::UAF
 {
 
 class FAnimNextModuleImpl : public IModuleInterface
diff --git a/Engine/Plugins/Experimental/AnimNext/Source/AnimNext/Internal/AnimNextPool.h b/Engine/Plugins/Experimental/AnimNext/Source/AnimNext/Internal/AnimNextPool.h
--- a/Engine/Plugins/Experimental/AnimNext/Source/AnimNext/Internal/AnimNextPool.h
+++ b/Engine/Plugins/Experimental/AnimNext/Source/AnimNext/Internal/AnimNextPool.h
@@ -7,7 +7,7 @@
 #include "Misc/TVariant.h"
 #include "AnimNextPoolHandle.h"
 
-namespace UE::AnimNext
+namespace UE::UAF
 {
 
 template<typename ElementType>
diff --git a/Engine/Plugins/Experimental/AnimNext/Source/AnimNext/Internal/Component/AnimNextWorldSubsystem.h b/Engine/Plugins/Experimental/AnimNext/Source/AnimNext/Internal/Component/AnimNextWorldSubsystem.h
--- a/Engine/Plugins/Experimental/AnimNext/Source/AnimNext/Internal/Component/AnimNextWorldSubsystem.h
+++ b/Engine/Plugins/Experimental/AnimNext/Source/AnimNext/Internal/Component/AnimNextWorldSubsystem.h
@@ -17,13 +17,13 @@
 class UAnimNextComponent;
 struct FRigUnit_AnimNextInitializeEvent;
 
-namespace UE::AnimNext
+namespace UE::UAF
 {
 	struct FModuleEventTickFunction;
 	struct FInjectionRequest;
 }
 
-namespace UE::AnimNext
+namespace UE::UAF
 {
 
 // A queued action to complete next frame
@@ -59,7 +59,7 @@ struct FModulePendingAction
 		Payload.Set<PayloadType>(InPayload);
 	}
 
-	TVariant<FModulePayloadNone, bool, UE::AnimNext::FModuleHandle> Payload;
+	TVariant<FModulePayloadNone, bool, UE::UAF::FModuleHandle> Payload;
 
 	FModuleHandle Handle;
 
@@ -97,7 +97,7 @@ private:
 	// UWorldSubsystem interface
 	UE_API virtual bool DoesSupportWorldType(const EWorldType::Type WorldType) const override;
 
-	UE_API bool IsValidHandle(UE::AnimNext::FModuleHandle InHandle) const;
+	UE_API bool IsValidHandle(UE::UAF::FModuleHandle InHandle) const;
 
 	UE_API void FlushPendingActions();
 
@@ -108,23 +108,23 @@ protected:
 #endif
 
 	// Register a handle to the subsystem
-	UE_API void RegisterHandle(UE::AnimNext::FModuleHandle& InOutHandle, UAnimNextModule* InModule, UObject* InObject, EAnimNextModuleInitMethod InitMethod);
+	UE_API void RegisterHandle(UE::UAF::FModuleHandle& InOutHandle, UAnimNextModule* InModule, UObject* InObject, EAnimNextModuleInitMethod InitMethod);
 
 	// Unregister a handle from the subsystem
-	UE_API void UnregisterHandle(UE::AnimNext::FModuleHandle& InOutHandle);
+	UE_API void UnregisterHandle(UE::UAF::FModuleHandle& InOutHandle);
 
 	// Returns whether the module represented by the supplied handle is enabled.
 	// If there are pending actions they will take precedent over the actual current state.
-	UE_API bool IsHandleEnabled(UE::AnimNext::FModuleHandle InHandle) const;
+	UE_API bool IsHandleEnabled(UE::UAF::FModuleHandle InHandle) const;
 
 	// Enables or disables the module represented by the supplied handle
 	// This operation is deferred until the next time the schedule ticks
-	UE_API void EnableHandle(UE::AnimNext::FModuleHandle InHandle, bool bInEnabled);
+	UE_API void EnableHandle(UE::UAF::FModuleHandle InHandle, bool bInEnabled);
 
 #if UE_ENABLE_DEBUG_DRAWING
 	// Enables or disables the module's debug drawing
 	// This operation is deferred until the next time the schedule ticks
-	UE_API void ShowDebugDrawingHandle(UE::AnimNext::FModuleHandle InHandle, bool bInShowDebugDrawing);
+	UE_API void ShowDebugDrawingHandle(UE::UAF::FModuleHandle InHandle, bool bInShowDebugDrawing);
 #endif
 
 	// Queue a task to run at a particular point in a schedule
@@ -132,15 +132,15 @@ protected:
 	// @param	InModuleEventName	The name of the event in the module to run the supplied task relative to. If this is NAME_None, then the first valid event will be used.
 	// @param	InTaskFunction		The function to run
 	// @param	InLocation			Where to run the task, before or after
-	UE_API void QueueTaskHandle(UE::AnimNext::FModuleHandle InHandle, FName InModuleEventName, TUniqueFunction<void(const UE::AnimNext::FModuleTaskContext&)>&& InTaskFunction, UE::AnimNext::ETaskRunLocation InLocation = UE::AnimNext::ETaskRunLocation::Before);
+	UE_API void QueueTaskHandle(UE::UAF::FModuleHandle InHandle, FName InModuleEventName, TUniqueFunction<void(const UE::UAF::FModuleTaskContext&)>&& InTaskFunction, UE::UAF::ETaskRunLocation InLocation = UE::UAF::ETaskRunLocation::Before);
 
 	// Queue an input trait event on the supplied handle
-	UE_API void QueueInputTraitEventHandle(UE::AnimNext::FModuleHandle InHandle, FAnimNextTraitEventPtr Event);
+	UE_API void QueueInputTraitEventHandle(UE::UAF::FModuleHandle InHandle, FAnimNextTraitEventPtr Event);
 
 	// Find the const tick function for the specified event. Useful to analyze/log prerequisites. Use AddDependencyHandle, RemoveDependencyHandle to actually modify dependencies
 	// @param	InHandle			The handle of the module
 	// @param	InEventName			The event associated to the wanted tick function
-	UE_API const FTickFunction* FindTickFunctionHandle(UE::AnimNext::FModuleHandle InHandle, FName InEventName) const;
+	UE_API const FTickFunction* FindTickFunctionHandle(UE::UAF::FModuleHandle InHandle, FName InEventName) const;
 
 	// Add a dependency on a tick function to the specified event
 	// @param	InHandle			The handle of the module whose dependencies are being modified
@@ -148,7 +148,7 @@ protected:
 	// @param	InTickFunction		The tick function to depend on
 	// @param	InEventName			The event to add the dependency to
 	// @param	InDependency		The kind of dependency to add
-	UE_API void AddDependencyHandle(UE::AnimNext::FModuleHandle InHandle, UObject* InObject, FTickFunction& InTickFunction, FName InEventName, EDependency InDependency);
+	UE_API void AddDependencyHandle(UE::UAF::FModuleHandle InHandle, UObject* InObject, FTickFunction& InTickFunction, FName InEventName, EDependency InDependency);
 	
 	// Remove a dependency on a tick function from the specified event
 	// @param	InHandle			The handle of the module whose dependencies are being modified
@@ -156,7 +156,7 @@ protected:
 	// @param	InTickFunction		The tick function that was depended on
 	// @param	InEventName			The event to remove the dependency from
 	// @param	InDependency		The kind of dependency to remove
-	UE_API void RemoveDependencyHandle(UE::AnimNext::FModuleHandle InHandle, UObject* InObject, FTickFunction& InTickFunction, FName InEventName, EDependency InDependency);
+	UE_API void RemoveDependencyHandle(UE::UAF::FModuleHandle InHandle, UObject* InObject, FTickFunction& InTickFunction, FName InEventName, EDependency InDependency);
 
 	// Add a dependency on a module event to the specified module event
 	// @param	InHandle			The handle of the module whose dependencies are being modified
@@ -164,7 +164,7 @@ protected:
 	// @param	OtherHandle			The handle of the other module whose dependencies are being modified
 	// @param	OtherEventName		The other module's event to add the dependency to/from
 	// @param	InDependency		The kind of dependency to add
-	UE_API void AddModuleEventDependencyHandle(UE::AnimNext::FModuleHandle InHandle, FName InEventName, UE::AnimNext::FModuleHandle OtherHandle, FName OtherEventName, EDependency InDependency);
+	UE_API void AddModuleEventDependencyHandle(UE::UAF::FModuleHandle InHandle, FName InEventName, UE::UAF::FModuleHandle OtherHandle, FName OtherEventName, EDependency InDependency);
 
 	// Remove a dependency on a module event from the specified module event
 	// @param	InHandle			The handle of the module whose dependencies are being modified
@@ -172,7 +172,7 @@ protected:
 	// @param	OtherHandle			The handle of the other module whose dependencies are being modified
 	// @param	OtherEventName		The other module's event to remove the dependency to/from
 	// @param	InDependency		The kind of dependency to remove
-	UE_API void RemoveModuleEventDependencyHandle(UE::AnimNext::FModuleHandle InHandle, FName InEventName, UE::AnimNext::FModuleHandle OtherHandle, FName OtherEventName, EDependency InDependency);
+	UE_API void RemoveModuleEventDependencyHandle(UE::UAF::FModuleHandle InHandle, FName InEventName, UE::UAF::FModuleHandle OtherHandle, FName OtherEventName, EDependency InDependency);
 
 	// Set the value of the specified variable
 	// @param	InHandle			The handle of the module whose variable we are setting
@@ -180,7 +180,7 @@ protected:
 	// @param	InType				The type of the variable we want to set
 	// @param	InData				The data to set the variable with
 	// @return true if the variable was set successfully
-	UE_API EPropertyBagResult SetVariableHandle(UE::AnimNext::FModuleHandle InHandle, const FAnimNextVariableReference& InVariable, const FAnimNextParamType& InType, TConstArrayView<uint8> InData);
+	UE_API EPropertyBagResult SetVariableHandle(UE::UAF::FModuleHandle InHandle, const FAnimNextVariableReference& InVariable, const FAnimNextParamType& InType, TConstArrayView<uint8> InData);
 
 	// Accesses the variable of the specified name for writing.
 	// General read access is not provided via this API due to the current double-buffering strategy used to communicate variable writes to worker threads.
@@ -190,18 +190,18 @@ protected:
 	// @param	InType				The type of the variable we want to access
 	// @param	InFunction			Function that will be called to allow modification of the variables
 	// @return true if the variable was accessed successfully
-	UE_API EPropertyBagResult WriteVariableHandle(UE::AnimNext::FModuleHandle InHandle, const FAnimNextVariableReference& InVariable, const FAnimNextParamType& InType, TFunctionRef<void(TArrayView<uint8>)> InFunction);
+	UE_API EPropertyBagResult WriteVariableHandle(UE::UAF::FModuleHandle InHandle, const FAnimNextVariableReference& InVariable, const FAnimNextParamType& InType, TFunctionRef<void(TArrayView<uint8>)> InFunction);
 
 protected:
 	friend struct FRigUnit_AnimNextInitializeEvent;
-	friend struct UE::AnimNext::FModuleEventTickFunction;
-	friend struct UE::AnimNext::FInjectionRequest;
+	friend struct UE::UAF::FModuleEventTickFunction;
+	friend struct UE::UAF::FInjectionRequest;
 
 	// Currently running instances, pooled
-	UE::AnimNext::TPool<FAnimNextModuleInstance> Instances;
+	UE::UAF::TPool<FAnimNextModuleInstance> Instances;
 
 	// Queued actions
-	TArray<UE::AnimNext::FModulePendingAction> PendingActions;
+	TArray<UE::UAF::FModulePendingAction> PendingActions;
 
 	// Locks for concurrent modifications
 	FRWLock InstancesLock;
diff --git a/Engine/Plugins/Experimental/AnimNext/Source/AnimNext/Internal/Component/SkinnedMeshComponentExtensions.h b/Engine/Plugins/Experimental/AnimNext/Source/AnimNext/Internal/Component/SkinnedMeshComponentExtensions.h
--- a/Engine/Plugins/Experimental/AnimNext/Source/AnimNext/Internal/Component/SkinnedMeshComponentExtensions.h
+++ b/Engine/Plugins/Experimental/AnimNext/Source/AnimNext/Internal/Component/SkinnedMeshComponentExtensions.h
@@ -17,7 +17,7 @@ struct FSkinnedMeshComponentExtensions
 	// Converts local to component space, flips buffers, updates bounds and dispatches to renderer
 	static ANIMNEXT_API void CompleteAndDispatch(
 		USkinnedMeshComponent* InComponent,
-		const UE::AnimNext::FLODPoseHeap& InLODPose,
+		const UE::UAF::FLODPoseHeap& InLODPose,
 		TUniqueFunction<void(USkinnedMeshComponent*)>&& InGameThreadCallback);
 };
 
diff --git a/Engine/Plugins/Experimental/AnimNext/Source/AnimNext/Internal/Graph/AnimNext_LODPose.h b/Engine/Plugins/Experimental/AnimNext/Source/AnimNext/Internal/Graph/AnimNext_LODPose.h
--- a/Engine/Plugins/Experimental/AnimNext/Source/AnimNext/Internal/Graph/AnimNext_LODPose.h
+++ b/Engine/Plugins/Experimental/AnimNext/Source/AnimNext/Internal/Graph/AnimNext_LODPose.h
@@ -17,12 +17,12 @@ struct FAnimNextGraphReferencePose
 
 	FAnimNextGraphReferencePose() = default;
 
-	explicit FAnimNextGraphReferencePose(UE::AnimNext::FDataHandle& InReferencePose)
+	explicit FAnimNextGraphReferencePose(UE::UAF::FDataHandle& InReferencePose)
 		: ReferencePose(InReferencePose)
 	{
 	}
 
-	UE::AnimNext::FDataHandle ReferencePose;
+	UE::UAF::FDataHandle ReferencePose;
 };
 
 USTRUCT(BlueprintType, meta = (DisplayName = "LODPose"))
@@ -31,7 +31,7 @@ struct FAnimNextGraphLODPose
 	GENERATED_BODY()
 
 	// Joint transforms
-	UE::AnimNext::FLODPoseHeap LODPose;
+	UE::UAF::FLODPoseHeap LODPose;
 
 	// Float curves
 	FBlendedHeapCurve Curves;
diff --git a/Engine/Plugins/Experimental/AnimNext/Source/AnimNext/Internal/Module/AnimNextModule.h b/Engine/Plugins/Experimental/AnimNext/Source/AnimNext/Internal/Module/AnimNextModule.h
--- a/Engine/Plugins/Experimental/AnimNext/Source/AnimNext/Internal/Module/AnimNextModule.h
+++ b/Engine/Plugins/Experimental/AnimNext/Source/AnimNext/Internal/Module/AnimNextModule.h
@@ -24,7 +24,7 @@ struct FAnimNextEditorParam;
 struct FAnimNextParam;
 struct FAnimNextModuleInstance;
 
-namespace UE::AnimNext
+namespace UE::UAF
 {
 	struct FContext;
 	struct FExecutionContext;
@@ -33,12 +33,12 @@ namespace UE::AnimNext
 	struct FParametersProxy;
 }
 
-namespace UE::AnimNext::UncookedOnly
+namespace UE::UAF::UncookedOnly
 {
 	struct FUtils;
 }
 
-namespace UE::AnimNext::Editor
+namespace UE::UAF::Editor
 {
 	class FModuleEditor;
 	class FVariableCustomization;
@@ -60,16 +60,16 @@ public:
 protected:
 	friend class UAnimNextModuleFactory;
 	friend class UAnimNextModule_EditorData;
-	friend struct UE::AnimNext::UncookedOnly::FUtils;
-	friend class UE::AnimNext::Editor::FModuleEditor;
-	friend struct UE::AnimNext::FTestUtils;
+	friend struct UE::UAF::UncookedOnly::FUtils;
+	friend class UE::UAF::Editor::FModuleEditor;
+	friend struct UE::UAF::FTestUtils;
 	friend FAnimNextGraphInstance;
 	friend class UAnimGraphNode_AnimNextGraph;
-	friend UE::AnimNext::FExecutionContext;
+	friend UE::UAF::FExecutionContext;
 	friend struct FAnimNextScheduleGraphTask;
-	friend UE::AnimNext::FAnimNextModuleImpl;
-	friend class UE::AnimNext::Editor::FVariableCustomization;
-	friend struct UE::AnimNext::FParametersProxy;
+	friend UE::UAF::FAnimNextModuleImpl;
+	friend class UE::UAF::Editor::FVariableCustomization;
+	friend struct UE::UAF::FParametersProxy;
 	friend FAnimNextModuleInstance;
 
 	// All components that are required on startup for this module
diff --git a/Engine/Plugins/Experimental/AnimNext/Source/AnimNext/Internal/Module/AnimNextModuleInstance.h b/Engine/Plugins/Experimental/AnimNext/Source/AnimNext/Internal/Module/AnimNextModuleInstance.h
--- a/Engine/Plugins/Experimental/AnimNext/Source/AnimNext/Internal/Module/AnimNextModuleInstance.h
+++ b/Engine/Plugins/Experimental/AnimNext/Source/AnimNext/Internal/Module/AnimNextModuleInstance.h
@@ -26,23 +26,23 @@ struct FRigUnit_CopyModuleProxyVariables;
 struct FAnimNextNotifyDispatcherComponent;
 struct FAnimNextGraphInstance;
 
-namespace UE::AnimNext
+namespace UE::UAF
 {
 	struct FModuleEventTickFunction;
 	struct FModuleEndTickFunction;
 }
 
-namespace UE::AnimNext::Debug
+namespace UE::UAF::Debug
 {
 	struct FDebugDraw;
 }
 
-namespace UE::AnimNext
+namespace UE::UAF
 {
 	using ModuleInstanceComponentMapType = TMap<FName, TInstancedStruct<FUAFModuleInstanceComponent>>;
 }
 
-namespace UE::AnimNext::Tests
+namespace UE::UAF::Tests
 {
 	class FVariables_Shared;
 }
@@ -57,7 +57,7 @@ struct FAnimNextModuleInstance : public FUAFAssetInstance
 	ANIMNEXT_API FAnimNextModuleInstance(
 		UAnimNextModule* InModule,
 		UObject* InObject,
-		UE::AnimNext::TPool<FAnimNextModuleInstance>* InPool,
+		UE::UAF::TPool<FAnimNextModuleInstance>* InPool,
 		EAnimNextModuleInitMethod InInitMethod);
 	ANIMNEXT_API ~FAnimNextModuleInstance();
 
@@ -93,11 +93,11 @@ struct FAnimNextModuleInstance : public FUAFAssetInstance
 	ANIMNEXT_API static void RunTaskOnGameThread(TUniqueFunction<void(void)>&& InFunction);
 
 	// Find the tick function for the specified event
-	UE::AnimNext::FModuleEventTickFunction* FindTickFunctionByName(FName InEventName);
-	const UE::AnimNext::FModuleEventTickFunction* FindTickFunctionByName(FName InEventName) const;
+	UE::UAF::FModuleEventTickFunction* FindTickFunctionByName(FName InEventName);
+	const UE::UAF::FModuleEventTickFunction* FindTickFunctionByName(FName InEventName) const;
 
 	// Find the first 'user' tick function
-	ANIMNEXT_API UE::AnimNext::FModuleEventTickFunction* FindFirstUserTickFunction();
+	ANIMNEXT_API UE::UAF::FModuleEventTickFunction* FindFirstUserTickFunction();
 
 	// Run the specified RigVM event
 	ANIMNEXT_API void RunRigVMEvent(FName InEventName, float InDeltaTime);
@@ -106,22 +106,22 @@ struct FAnimNextModuleInstance : public FUAFAssetInstance
 	EWorldType::Type GetWorldType() const { return WorldType; }
 
 	// Get tick functions for this module instance
-	ANIMNEXT_API TArrayView<UE::AnimNext::FModuleEventTickFunction> GetTickFunctions();
+	ANIMNEXT_API TArrayView<UE::UAF::FModuleEventTickFunction> GetTickFunctions();
 
-	UE::AnimNext::FModuleHandle GetHandle() const { return Handle; }
+	UE::UAF::FModuleHandle GetHandle() const { return Handle; }
 
 	// Queue a task to run at a particular module event
 	// @param	InModuleEvent		The event in the module to run the supplied task relative to. If this is NAME_None or not found then the first 'user' event will be used.
 	// @param	InTaskFunction		The function to run
 	// @param	InLocation			Where to run the task, before or after
-	ANIMNEXT_API void QueueTask(FName InEventName, TUniqueFunction<void(const UE::AnimNext::FModuleTaskContext&)>&& InTaskFunction, UE::AnimNext::ETaskRunLocation InLocation = UE::AnimNext::ETaskRunLocation::Before);
+	ANIMNEXT_API void QueueTask(FName InEventName, TUniqueFunction<void(const UE::UAF::FModuleTaskContext&)>&& InTaskFunction, UE::UAF::ETaskRunLocation InLocation = UE::UAF::ETaskRunLocation::Before);
 
 	// Queue a task to run at a particular module event on some other module
 	// @param	InModuleHandle		The handle for the module to queue the task on
 	// @param	InModuleEventName	The event in the module to run the supplied task relative to. If this is NAME_None or not found then the first user 'event' will be used.
 	// @param	InTaskFunction		The function to run
 	// @param	InLocation			Where to run the task, before or after
-	ANIMNEXT_API void QueueTaskOnOtherModule(const UE::AnimNext::FModuleHandle InOtherModuleHandle, FName InEventName, TUniqueFunction<void(const UE::AnimNext::FModuleTaskContext&)>&& InTaskFunction, UE::AnimNext::ETaskRunLocation InLocation = UE::AnimNext::ETaskRunLocation::Before);
+	ANIMNEXT_API void QueueTaskOnOtherModule(const UE::UAF::FModuleHandle InOtherModuleHandle, FName InEventName, TUniqueFunction<void(const UE::UAF::FModuleTaskContext&)>&& InTaskFunction, UE::UAF::ETaskRunLocation InLocation = UE::UAF::ETaskRunLocation::Before);
 
 private:
 	// Setup the entry
@@ -151,7 +151,7 @@ private:
 
 	// Handles trait events at the module level.
 	// Called in the module end tick function before any events are dispatched.
-	void RaiseTraitEvents(const UE::AnimNext::FTraitEventList& EventList);
+	void RaiseTraitEvents(const UE::UAF::FTraitEventList& EventList);
 
 	// Set the value of the specified variable via the proxy double-buffer. Called from any thread.
 	EPropertyBagResult SetProxyVariable(const FAnimNextVariableReference& InVariable, const FAnimNextParamType& InType, TConstArrayView<uint8> InData);
@@ -160,8 +160,8 @@ private:
 	EPropertyBagResult WriteProxyVariable(const FAnimNextVariableReference& InVariable, const FAnimNextParamType& InType, TFunctionRef<void(TArrayView<uint8>)> InFunction);
 
 private:
-	friend UE::AnimNext::FModuleEventTickFunction;
-	friend UE::AnimNext::FModuleEndTickFunction;
+	friend UE::UAF::FModuleEventTickFunction;
+	friend UE::UAF::FModuleEndTickFunction;
 	friend FRigUnit_AnimNextRunAnimationGraph_v1;
 	friend FRigUnit_AnimNextRunAnimationGraph_v2;
 	friend UAnimNextWorldSubsystem;
@@ -169,26 +169,26 @@ private:
 	friend FRigUnit_CopyModuleProxyVariables;
 	friend FAnimNextNotifyDispatcherComponent;
 	friend FAnimNextGraphInstance;
-	friend UE::AnimNext::Tests::FVariables_Shared;
+	friend UE::UAF::Tests::FVariables_Shared;
 
 	// Object this entry is bound to
 	UPROPERTY(Transient)
 	TObjectPtr<UObject> Object = nullptr;
 
 	// The pool that this module instance exists in
-	UE::AnimNext::TPool<FAnimNextModuleInstance>* Pool = nullptr;
+	UE::UAF::TPool<FAnimNextModuleInstance>* Pool = nullptr;
 
 	// Copy of the handle that represents this entry to client systems
-	UE::AnimNext::FModuleHandle Handle;
+	UE::UAF::FModuleHandle Handle;
 
 	// Pre-allocated graph of tick functions
-	TArray<UE::AnimNext::FModuleEventTickFunction> TickFunctions;
+	TArray<UE::UAF::FModuleEventTickFunction> TickFunctions;
 
 	// Input event list to be processed on the next update
-	UE::AnimNext::FTraitEventList InputEventList;
+	UE::UAF::FTraitEventList InputEventList;
 
 	// Output event list to be processed at the end of the schedule tick
-	UE::AnimNext::FTraitEventList OutputEventList;
+	UE::UAF::FTraitEventList OutputEventList;
 
 	// Lock to ensure event list actions are thread safe
 	FRWLock EventListLock;
@@ -206,7 +206,7 @@ private:
 #if UE_ENABLE_DEBUG_DRAWING
 	// Debug draw support
 	// TODO: Move this to a 'module component' once that code is checked in
-	TUniquePtr<UE::AnimNext::Debug::FDebugDraw> DebugDraw;
+	TUniquePtr<UE::UAF::Debug::FDebugDraw> DebugDraw;
 #endif
 
 	enum class EInitState : uint8
diff --git a/Engine/Plugins/Experimental/AnimNext/Source/AnimNext/Internal/Module/ModuleTickFunction.h b/Engine/Plugins/Experimental/AnimNext/Source/AnimNext/Internal/Module/ModuleTickFunction.h
--- a/Engine/Plugins/Experimental/AnimNext/Source/AnimNext/Internal/Module/ModuleTickFunction.h
+++ b/Engine/Plugins/Experimental/AnimNext/Source/AnimNext/Internal/Module/ModuleTickFunction.h
@@ -9,13 +9,13 @@
 
 struct FAnimNextModuleInstance;
 
-namespace UE::AnimNext
+namespace UE::UAF
 {
 	struct FScheduleContext;
 	struct FModuleTaskContext;
 }
 
-namespace UE::AnimNext
+namespace UE::UAF
 {
 
 // Multicast delegate called before module events
diff --git a/Engine/Plugins/Experimental/AnimNext/Source/AnimNext/Internal/Module/RigUnit_AnimNextModuleEvents.h b/Engine/Plugins/Experimental/AnimNext/Source/AnimNext/Internal/Module/RigUnit_AnimNextModuleEvents.h
--- a/Engine/Plugins/Experimental/AnimNext/Source/AnimNext/Internal/Module/RigUnit_AnimNextModuleEvents.h
+++ b/Engine/Plugins/Experimental/AnimNext/Source/AnimNext/Internal/Module/RigUnit_AnimNextModuleEvents.h
@@ -9,7 +9,7 @@
 
 #define UE_API ANIMNEXT_API
 
-namespace UE::AnimNext
+namespace UE::UAF
 {
 
 // Phase is used as a general ordering constraint on event execution
@@ -37,10 +37,10 @@ struct FRigUnit_AnimNextModuleEventBase : public FRigUnit_AnimNextBase
 	virtual bool CanOnlyExistOnce() const override { return true; }
 
 	// Get the general ordering phase of this event, used for linearization
-	virtual UE::AnimNext::EModuleEventPhase GetEventPhase() const { return UE::AnimNext::EModuleEventPhase::Execute; }
+	virtual UE::UAF::EModuleEventPhase GetEventPhase() const { return UE::UAF::EModuleEventPhase::Execute; }
 
 	// Overriden in derived classes to provide binding function
-	virtual UE::AnimNext::FModuleEventBindingFunction GetBindingFunction() const { return [](const UE::AnimNext::FTickFunctionBindingContext& InContext, FTickFunction& InTickFunction){}; }
+	virtual UE::UAF::FModuleEventBindingFunction GetBindingFunction() const { return [](const UE::UAF::FTickFunctionBindingContext& InContext, FTickFunction& InTickFunction){}; }
 
 	// Get the tick group for this event
 	virtual ETickingGroup GetTickGroup() const { return TG_PrePhysics; }
@@ -80,8 +80,8 @@ struct FRigUnit_AnimNextExecuteBindings_GT : public FRigUnit_AnimNextModuleEvent
 	virtual bool IsGameThreadTask() const override final { return true; }
 
 	// FRigUnit_AnimNextModuleEventBase interface
-	virtual UE::AnimNext::EModuleEventPhase GetEventPhase() const override final { return UE::AnimNext::EModuleEventPhase::PreExecute; }
-	UE_API virtual UE::AnimNext::FModuleEventBindingFunction GetBindingFunction() const override final;
+	virtual UE::UAF::EModuleEventPhase GetEventPhase() const override final { return UE::UAF::EModuleEventPhase::PreExecute; }
+	UE_API virtual UE::UAF::FModuleEventBindingFunction GetBindingFunction() const override final;
 };
 
 /** Synthetic event injected by the compiler to process any variable bindings on a worker thread, not user instantiated */
@@ -102,8 +102,8 @@ struct FRigUnit_AnimNextExecuteBindings_WT : public FRigUnit_AnimNextModuleEvent
 	virtual bool IsTask() const override final { return false; }
 
 	// FRigUnit_AnimNextModuleEventBase interface
-	virtual UE::AnimNext::EModuleEventPhase GetEventPhase() const override final { return UE::AnimNext::EModuleEventPhase::PreExecute; }
-	UE_API virtual UE::AnimNext::FModuleEventBindingFunction GetBindingFunction() const override final;
+	virtual UE::UAF::EModuleEventPhase GetEventPhase() const override final { return UE::UAF::EModuleEventPhase::PreExecute; }
+	UE_API virtual UE::UAF::FModuleEventBindingFunction GetBindingFunction() const override final;
 };
 
 /** Schedule event called to set up a module */
@@ -131,8 +131,8 @@ struct FRigUnit_AnimNextUserEvent : public FRigUnit_AnimNextModuleEventBase
 	virtual FString GetUnitLabel() const override { return GetEventName().ToString(); }
 	UE_API virtual FString GetUnitSubTitle() const override;
 	virtual bool CanOnlyExistOnce() const override { return false; }
-	virtual UE::AnimNext::EModuleEventPhase GetEventPhase() const override final { return UE::AnimNext::EModuleEventPhase::Execute; }
-	UE_API virtual UE::AnimNext::FModuleEventBindingFunction GetBindingFunction() const override final;
+	virtual UE::UAF::EModuleEventPhase GetEventPhase() const override final { return UE::UAF::EModuleEventPhase::Execute; }
+	UE_API virtual UE::UAF::FModuleEventBindingFunction GetBindingFunction() const override final;
 	virtual ETickingGroup GetTickGroup() const override final { return TickGroup; }
 	virtual int32 GetSortOrder() const override final { return SortOrder; }
 
diff --git a/Engine/Plugins/Experimental/AnimNext/Source/AnimNext/Internal/Param/RigVMDispatch_GetLayerParameter.h b/Engine/Plugins/Experimental/AnimNext/Source/AnimNext/Internal/Param/RigVMDispatch_GetLayerParameter.h
--- a/Engine/Plugins/Experimental/AnimNext/Source/AnimNext/Internal/Param/RigVMDispatch_GetLayerParameter.h
+++ b/Engine/Plugins/Experimental/AnimNext/Source/AnimNext/Internal/Param/RigVMDispatch_GetLayerParameter.h
@@ -8,7 +8,7 @@
 
 #define UE_API ANIMNEXT_API
 
-namespace UE::AnimNext::UncookedOnly
+namespace UE::UAF::UncookedOnly
 {
 	struct FUtils;
 }
@@ -24,7 +24,7 @@ struct FRigVMDispatch_GetLayerParameter : public FRigVMDispatchFactory
 	UE_API FRigVMDispatch_GetLayerParameter();
 
 private:
-	friend struct UE::AnimNext::UncookedOnly::FUtils;
+	friend struct UE::UAF::UncookedOnly::FUtils;
 
 	virtual UScriptStruct* GetExecuteContextStruct() const override { return FAnimNextExecuteContext::StaticStruct(); }
 	UE_API virtual FName GetArgumentNameForOperandIndex(int32 InOperandIndex, int32 InTotalOperands, FRigVMRegistryHandle& InRegistry) const override;
diff --git a/Engine/Plugins/Experimental/AnimNext/Source/AnimNext/Internal/Param/RigVMDispatch_GetParameter.h b/Engine/Plugins/Experimental/AnimNext/Source/AnimNext/Internal/Param/RigVMDispatch_GetParameter.h
--- a/Engine/Plugins/Experimental/AnimNext/Source/AnimNext/Internal/Param/RigVMDispatch_GetParameter.h
+++ b/Engine/Plugins/Experimental/AnimNext/Source/AnimNext/Internal/Param/RigVMDispatch_GetParameter.h
@@ -8,7 +8,7 @@
 
 #define UE_API ANIMNEXT_API
 
-namespace UE::AnimNext::UncookedOnly
+namespace UE::UAF::UncookedOnly
 {
 	struct FUtils;
 }
@@ -27,7 +27,7 @@ struct FRigVMDispatch_GetParameter : public FRigVMDispatchFactory
 	static UE_API const FName ValueName;
 
 private:
-	friend struct UE::AnimNext::UncookedOnly::FUtils;
+	friend struct UE::UAF::UncookedOnly::FUtils;
 
 	virtual UScriptStruct* GetExecuteContextStruct() const { return FAnimNextExecuteContext::StaticStruct(); }
 	UE_API virtual FName GetArgumentNameForOperandIndex(int32 InOperandIndex, int32 InTotalOperands, FRigVMRegistryHandle& InRegistry) const override;
diff --git a/Engine/Plugins/Experimental/AnimNext/Source/AnimNext/Internal/Param/RigVMDispatch_GetScopedParameter.h b/Engine/Plugins/Experimental/AnimNext/Source/AnimNext/Internal/Param/RigVMDispatch_GetScopedParameter.h
--- a/Engine/Plugins/Experimental/AnimNext/Source/AnimNext/Internal/Param/RigVMDispatch_GetScopedParameter.h
+++ b/Engine/Plugins/Experimental/AnimNext/Source/AnimNext/Internal/Param/RigVMDispatch_GetScopedParameter.h
@@ -8,12 +8,12 @@
 
 #define UE_API ANIMNEXT_API
 
-namespace UE::AnimNext::Editor
+namespace UE::UAF::Editor
 {
 	class SGraphPinParam;
 }
 
-namespace UE::AnimNext::UncookedOnly
+namespace UE::UAF::UncookedOnly
 {
 	struct FUtils;
 }
@@ -29,8 +29,8 @@ struct FRigVMDispatch_GetScopedParameter : public FRigVMDispatchFactory
 	UE_API FRigVMDispatch_GetScopedParameter();
 
 private:
-	friend struct UE::AnimNext::UncookedOnly::FUtils;
-	friend class UE::AnimNext::Editor::SGraphPinParam;
+	friend struct UE::UAF::UncookedOnly::FUtils;
+	friend class UE::UAF::Editor::SGraphPinParam;
 
 	virtual UScriptStruct* GetExecuteContextStruct() const { return FAnimNextExecuteContext::StaticStruct(); }
 	UE_API virtual FName GetArgumentNameForOperandIndex(int32 InOperandIndex, int32 InTotalOperands, FRigVMRegistryHandle& InRegistry) const override;
diff --git a/Engine/Plugins/Experimental/AnimNext/Source/AnimNext/Internal/RigVMRuntimeDataRegistry.h b/Engine/Plugins/Experimental/AnimNext/Source/AnimNext/Internal/RigVMRuntimeDataRegistry.h
--- a/Engine/Plugins/Experimental/AnimNext/Source/AnimNext/Internal/RigVMRuntimeDataRegistry.h
+++ b/Engine/Plugins/Experimental/AnimNext/Source/AnimNext/Internal/RigVMRuntimeDataRegistry.h
@@ -8,7 +8,7 @@
 #define UE_API ANIMNEXT_API
 
 
-namespace UE::AnimNext
+namespace UE::UAF
 {
 
 using FRigVMRuntimeDataID = TObjectKey<URigVM>;
@@ -83,6 +83,6 @@ private:
 	friend class FAnimNextModuleImpl;
 };
 
-} // end namespace UE::AnimNext
+} // end namespace UE::UAF
 
 #undef UE_API
diff --git a/Engine/Plugins/Experimental/AnimNext/Source/AnimNext/Internal/Variables/AnimNextVariableOverridesCollection.h b/Engine/Plugins/Experimental/AnimNext/Source/AnimNext/Internal/Variables/AnimNextVariableOverridesCollection.h
--- a/Engine/Plugins/Experimental/AnimNext/Source/AnimNext/Internal/Variables/AnimNextVariableOverridesCollection.h
+++ b/Engine/Plugins/Experimental/AnimNext/Source/AnimNext/Internal/Variables/AnimNextVariableOverridesCollection.h
@@ -7,7 +7,7 @@
 struct FRigUnit_OverrideVariables;
 class FReferenceCollector;
 
-namespace UE::AnimNext
+namespace UE::UAF
 {
 	struct FVariableOverridesCollection;
 }
@@ -18,7 +18,7 @@ struct FAnimNextVariableOverridesCollection
 	GENERATED_BODY()
 
 	// Access the overrides for reading only
-	TWeakPtr<const UE::AnimNext::FVariableOverridesCollection> GetOverrides() const
+	TWeakPtr<const UE::UAF::FVariableOverridesCollection> GetOverrides() const
 	{
 		return Collection;
 	}
@@ -31,7 +31,7 @@ private:
 	friend FRigUnit_OverrideVariables;
 
 	// Overrides that this structure wraps - ownership is with a module/graph's execute context (held in work data)
-	TWeakPtr<UE::AnimNext::FVariableOverridesCollection> Collection;
+	TWeakPtr<UE::UAF::FVariableOverridesCollection> Collection;
 };
 
 template<>
diff --git a/Engine/Plugins/Experimental/AnimNext/Source/AnimNext/Internal/Variables/UAFInstanceVariableContainer.h b/Engine/Plugins/Experimental/AnimNext/Source/AnimNext/Internal/Variables/UAFInstanceVariableContainer.h
--- a/Engine/Plugins/Experimental/AnimNext/Source/AnimNext/Internal/Variables/UAFInstanceVariableContainer.h
+++ b/Engine/Plugins/Experimental/AnimNext/Source/AnimNext/Internal/Variables/UAFInstanceVariableContainer.h
@@ -16,7 +16,7 @@ struct FUAFInstanceVariableDataProxy;
 struct FAnimNextParamType;
 struct FUAFRigVMComponent;
 
-namespace UE::AnimNext
+namespace UE::UAF
 {
 	struct FVariableOverridesCollection;
 	struct FVariableOverrides;
@@ -47,14 +47,14 @@ private:
 	friend FUAFInstanceVariableDataProxy;
 	template<typename, ESPMode>
 	friend class SharedPointerInternals::TIntrusiveReferenceController;
-	friend UE::AnimNext::FAnimNextTrace;
+	friend UE::UAF::FAnimNextTrace;
 	friend FUAFRigVMComponent;
 
 	// Construct from struct data, copying the defaults. Optional overrides provided.
-	explicit FUAFInstanceVariableContainer(const FUAFAssetInstance& InHostInstance, const UScriptStruct* InStruct, const TSharedPtr<const UE::AnimNext::FVariableOverridesCollection>& InOverrides = nullptr, int32 InOverridesIndex = INDEX_NONE);
+	explicit FUAFInstanceVariableContainer(const FUAFAssetInstance& InHostInstance, const UScriptStruct* InStruct, const TSharedPtr<const UE::UAF::FVariableOverridesCollection>& InOverrides = nullptr, int32 InOverridesIndex = INDEX_NONE);
 
 	// Construct from an asset, copying the variables. Optional overrides provided.
-	explicit FUAFInstanceVariableContainer(const FUAFAssetInstance& InHostInstance, const UAnimNextRigVMAsset* InAsset, const TSharedPtr<const UE::AnimNext::FVariableOverridesCollection>& InOverrides = nullptr, int32 InOverridesIndex = INDEX_NONE);
+	explicit FUAFInstanceVariableContainer(const FUAFAssetInstance& InHostInstance, const UAnimNextRigVMAsset* InAsset, const TSharedPtr<const UE::UAF::FVariableOverridesCollection>& InOverrides = nullptr, int32 InOverridesIndex = INDEX_NONE);
 
 	// Gets the value of the specified variable. If the type does not match exactly then a conversion will be attempted.
 	EPropertyBagResult GetVariable(int32 InVariableIndex, const FAnimNextParamType& InType, TArrayView<uint8> OutResult) const;
@@ -102,7 +102,7 @@ private:
 	FVariablesContainerType VariablesContainer;
 
 	// Cached overrides collection, if any are to be applied
-	TWeakPtr<const UE::AnimNext::FVariableOverridesCollection> Overrides;
+	TWeakPtr<const UE::UAF::FVariableOverridesCollection> Overrides;
 
 	// Index into overrides collection to apply
 	int32 OverridesIndex = INDEX_NONE;
@@ -111,7 +111,7 @@ private:
 	TArray<uint8*> ResolvedOverrides;
 
 	// Reference to cached struct data, if required
-	TSharedPtr<UE::AnimNext::FStructDataCache> StructDataCache;
+	TSharedPtr<UE::UAF::FStructDataCache> StructDataCache;
 
 	// Cached variable count
 	int32 NumVariables = 0;
diff --git a/Engine/Plugins/Experimental/AnimNext/Source/AnimNext/Private/AnimNextDebugDraw.cpp b/Engine/Plugins/Experimental/AnimNext/Source/AnimNext/Private/AnimNextDebugDraw.cpp
--- a/Engine/Plugins/Experimental/AnimNext/Source/AnimNext/Private/AnimNextDebugDraw.cpp
+++ b/Engine/Plugins/Experimental/AnimNext/Source/AnimNext/Private/AnimNextDebugDraw.cpp
@@ -13,7 +13,7 @@
 #include "SceneView.h"
 #include "Module/AnimNextModuleInstance.h"
 
-namespace UE::AnimNext::Debug
+namespace UE::UAF::Debug
 {
 
 #if UE_ENABLE_DEBUG_DRAWING
diff --git a/Engine/Plugins/Experimental/AnimNext/Source/AnimNext/Private/AnimNextModuleImpl.cpp b/Engine/Plugins/Experimental/AnimNext/Source/AnimNext/Private/AnimNextModuleImpl.cpp
--- a/Engine/Plugins/Experimental/AnimNext/Source/AnimNext/Private/AnimNextModuleImpl.cpp
+++ b/Engine/Plugins/Experimental/AnimNext/Source/AnimNext/Private/AnimNextModuleImpl.cpp
@@ -31,7 +31,7 @@
 
 #define LOCTEXT_NAMESPACE "AnimNextModule"
 
-namespace UE::AnimNext
+namespace UE::UAF
 {
 #if ANIMNEXT_TRACE_ENABLED
 	FRewindDebuggerAnimNextRuntime GRewindDebuggerAnimNextRuntime;
@@ -125,4 +125,4 @@ namespace UE::AnimNext
 
 #undef LOCTEXT_NAMESPACE
 
-IMPLEMENT_MODULE(UE::AnimNext::FAnimNextModuleImpl, AnimNext)
+IMPLEMENT_MODULE(UE::UAF::FAnimNextModuleImpl, AnimNext)
diff --git a/Engine/Plugins/Experimental/AnimNext/Source/AnimNext/Private/AnimNextRigVMAsset.cpp b/Engine/Plugins/Experimental/AnimNext/Source/AnimNext/Private/AnimNextRigVMAsset.cpp
--- a/Engine/Plugins/Experimental/AnimNext/Source/AnimNext/Private/AnimNextRigVMAsset.cpp
+++ b/Engine/Plugins/Experimental/AnimNext/Source/AnimNext/Private/AnimNextRigVMAsset.cpp
@@ -12,7 +12,7 @@
 #include "Modules/ModuleManager.h"
 #include "UObject/AssetRegistryTagsContext.h"
 
-namespace UE::AnimNext::Private
+namespace UE::UAF::Private
 {
 #if WITH_EDITOR
 	static UAnimNextRigVMAsset::FOnCompileJobEvent GOnCompileJobStarted;
@@ -32,7 +32,7 @@ void UAnimNextRigVMAsset::BeginDestroy()
 
 	if (VM)
 	{
-		UE::AnimNext::FRigVMRuntimeDataRegistry::ReleaseAllVMRuntimeData(VM);
+		UE::UAF::FRigVMRuntimeDataRegistry::ReleaseAllVMRuntimeData(VM);
 	}
 }
 
@@ -128,7 +128,7 @@ TArray<FRigVMExternalVariable> UAnimNextRigVMAsset::GetExternalVariablesImpl(boo
 }
 
 
-void UAnimNextRigVMAsset::CallFunctionHandle(UE::AnimNext::FFunctionHandle InHandle, FRigVMExtendedExecuteContext& InContext, FInstancedPropertyBag& InArgs)
+void UAnimNextRigVMAsset::CallFunctionHandle(UE::UAF::FFunctionHandle InHandle, FRigVMExtendedExecuteContext& InContext, FInstancedPropertyBag& InArgs)
 {
 #if WITH_EDITOR
 	check(InHandle.IsValidForVM(VM->GetVMHash()));
@@ -173,9 +173,9 @@ void UAnimNextRigVMAsset::CallFunctionHandle(UE::AnimNext::FFunctionHandle InHan
 	}
 }
 
-UE::AnimNext::FFunctionHandle UAnimNextRigVMAsset::GetFunctionHandle(FName InEventName) const
+UE::UAF::FFunctionHandle UAnimNextRigVMAsset::GetFunctionHandle(FName InEventName) const
 {
-	using namespace UE::AnimNext;
+	using namespace UE::UAF;
 
 	for (int32 FunctionIndex = 0; FunctionIndex < FunctionData.Num(); ++FunctionIndex)
 	{
@@ -198,12 +198,12 @@ UE::AnimNext::FFunctionHandle UAnimNextRigVMAsset::GetFunctionHandle(FName InEve
 
 UAnimNextRigVMAsset::FOnCompileJobEvent& UAnimNextRigVMAsset::OnCompileJobStarted()
 {
-	return UE::AnimNext::Private::GOnCompileJobStarted;
+	return UE::UAF::Private::GOnCompileJobStarted;
 }
 
 UAnimNextRigVMAsset::FOnCompileJobEvent& UAnimNextRigVMAsset::OnCompileJobFinished()
 {
-	return UE::AnimNext::Private::GOnCompileJobFinished;
+	return UE::UAF::Private::GOnCompileJobFinished;
 }
 
 #endif
\ No newline at end of file
diff --git a/Engine/Plugins/Experimental/AnimNext/Source/AnimNext/Private/Component/AnimNextComponent.cpp b/Engine/Plugins/Experimental/AnimNext/Source/AnimNext/Private/Component/AnimNextComponent.cpp
--- a/Engine/Plugins/Experimental/AnimNext/Source/AnimNext/Private/Component/AnimNextComponent.cpp
+++ b/Engine/Plugins/Experimental/AnimNext/Source/AnimNext/Private/Component/AnimNextComponent.cpp
@@ -10,7 +10,7 @@
 
 void UAnimNextComponent::OnRegister()
 {
-	using namespace UE::AnimNext;
+	using namespace UE::UAF;
 
 	Super::OnRegister();
 
@@ -66,7 +66,7 @@ void UAnimNextComponent::BlueprintSetVariable(FName Name, int32 Value)
 
 DEFINE_FUNCTION(UAnimNextComponent::execBlueprintSetVariable)
 {
-	using namespace UE::AnimNext;
+	using namespace UE::UAF;
 
 	// Read wildcard Value input.
 	Stack.MostRecentPropertyAddress = nullptr;
@@ -161,7 +161,7 @@ void UAnimNextComponent::BlueprintSetVariableReference(const FAnimNextVariableRe
 
 DEFINE_FUNCTION(UAnimNextComponent::execBlueprintSetVariableReference)
 {
-	using namespace UE::AnimNext;
+	using namespace UE::UAF;
 
 	P_GET_STRUCT_REF(FAnimNextVariableReference, Variable);
 
@@ -301,9 +301,9 @@ void UAnimNextComponent::ShowDebugDrawing(bool bInShowDebugDrawing)
 #endif
 }
 
-void UAnimNextComponent::QueueTask(FName InModuleEventName, TUniqueFunction<void(const UE::AnimNext::FModuleTaskContext&)>&& InTaskFunction, UE::AnimNext::ETaskRunLocation InLocation)
+void UAnimNextComponent::QueueTask(FName InModuleEventName, TUniqueFunction<void(const UE::UAF::FModuleTaskContext&)>&& InTaskFunction, UE::UAF::ETaskRunLocation InLocation)
 {
-	using namespace UE::AnimNext;
+	using namespace UE::UAF;
 
 	if(Subsystem)
 	{
@@ -313,7 +313,7 @@ void UAnimNextComponent::QueueTask(FName InModuleEventName, TUniqueFunction<void
 
 void UAnimNextComponent::QueueInputTraitEvent(FAnimNextTraitEventPtr Event)
 {
-	using namespace UE::AnimNext;
+	using namespace UE::UAF;
 
 	if(Subsystem)
 	{
diff --git a/Engine/Plugins/Experimental/AnimNext/Source/AnimNext/Private/Component/AnimNextComponentWorldSubsystem.cpp b/Engine/Plugins/Experimental/AnimNext/Source/AnimNext/Private/Component/AnimNextComponentWorldSubsystem.cpp
--- a/Engine/Plugins/Experimental/AnimNext/Source/AnimNext/Private/Component/AnimNextComponentWorldSubsystem.cpp
+++ b/Engine/Plugins/Experimental/AnimNext/Source/AnimNext/Private/Component/AnimNextComponentWorldSubsystem.cpp
@@ -11,7 +11,7 @@
 
 void UAnimNextComponentWorldSubsystem::Register(UAnimNextComponent* InComponent)
 {
-	using namespace UE::AnimNext;
+	using namespace UE::UAF;
 	check(IsInGameThread());
 	check(InComponent);
 	RegisterHandle(InComponent->ModuleHandle, InComponent->Module, InComponent, InComponent->InitMethod);
@@ -19,7 +19,7 @@ void UAnimNextComponentWorldSubsystem::Register(UAnimNextComponent* InComponent)
 
 void UAnimNextComponentWorldSubsystem::Unregister(UAnimNextComponent* InComponent)
 {
-	using namespace UE::AnimNext;
+	using namespace UE::UAF;
 	check(IsInGameThread());
 	check(InComponent);
 	UnregisterHandle(InComponent->ModuleHandle);
@@ -27,7 +27,7 @@ void UAnimNextComponentWorldSubsystem::Unregister(UAnimNextComponent* InComponen
 
 void UAnimNextComponentWorldSubsystem::SetEnabled(UAnimNextComponent* InComponent, bool bInEnabled)
 {
-	using namespace UE::AnimNext;
+	using namespace UE::UAF;
 	check(IsInGameThread());
 	check(InComponent);
 	EnableHandle(InComponent->ModuleHandle, bInEnabled);
@@ -35,7 +35,7 @@ void UAnimNextComponentWorldSubsystem::SetEnabled(UAnimNextComponent* InComponen
 
 bool UAnimNextComponentWorldSubsystem::IsEnabled(const UAnimNextComponent* InComponent) const
 {
-	using namespace UE::AnimNext;
+	using namespace UE::UAF;
 	check(IsInGameThread());
 	check(InComponent);
 	return IsHandleEnabled(InComponent->ModuleHandle);
@@ -44,16 +44,16 @@ bool UAnimNextComponentWorldSubsystem::IsEnabled(const UAnimNextComponent* InCom
 #if UE_ENABLE_DEBUG_DRAWING
 void UAnimNextComponentWorldSubsystem::ShowDebugDrawing(UAnimNextComponent* InComponent, bool bInShowDebugDrawing)
 {
-	using namespace UE::AnimNext;
+	using namespace UE::UAF;
 	check(IsInGameThread());
 	check(InComponent);
 	ShowDebugDrawingHandle(InComponent->ModuleHandle, bInShowDebugDrawing);
 }
 #endif
 
-void UAnimNextComponentWorldSubsystem::QueueTask(UAnimNextComponent* InComponent, FName InModuleEventName, TUniqueFunction<void(const UE::AnimNext::FModuleTaskContext&)>&& InTaskFunction, UE::AnimNext::ETaskRunLocation InLocation)
+void UAnimNextComponentWorldSubsystem::QueueTask(UAnimNextComponent* InComponent, FName InModuleEventName, TUniqueFunction<void(const UE::UAF::FModuleTaskContext&)>&& InTaskFunction, UE::UAF::ETaskRunLocation InLocation)
 {
-	using namespace UE::AnimNext;
+	using namespace UE::UAF;
 	check(IsInGameThread());
 	check(InComponent);
 	QueueTaskHandle(InComponent->ModuleHandle, InModuleEventName, MoveTemp(InTaskFunction), InLocation);
@@ -61,7 +61,7 @@ void UAnimNextComponentWorldSubsystem::QueueTask(UAnimNextComponent* InComponent
 
 void UAnimNextComponentWorldSubsystem::QueueInputTraitEvent(UAnimNextComponent* InComponent, FAnimNextTraitEventPtr Event)
 {
-	using namespace UE::AnimNext;
+	using namespace UE::UAF;
 	check(IsInGameThread());
 	check(InComponent);
 	QueueInputTraitEventHandle(InComponent->ModuleHandle, Event);
@@ -69,7 +69,7 @@ void UAnimNextComponentWorldSubsystem::QueueInputTraitEvent(UAnimNextComponent*
 
 const FTickFunction* UAnimNextComponentWorldSubsystem::FindTickFunction(const UAnimNextComponent* InComponent, FName InEventName) const
 {
-	using namespace UE::AnimNext;
+	using namespace UE::UAF;
 	check(IsInGameThread());
 	check(InComponent);
 	return FindTickFunctionHandle(InComponent->ModuleHandle, InEventName);
@@ -77,7 +77,7 @@ const FTickFunction* UAnimNextComponentWorldSubsystem::FindTickFunction(const UA
 
 void UAnimNextComponentWorldSubsystem::AddDependency(UAnimNextComponent* InComponent, UObject* InObject, FTickFunction& InTickFunction, FName InEventName, EDependency InDependency)
 {
-	using namespace UE::AnimNext;
+	using namespace UE::UAF;
 	check(IsInGameThread());
 	check(InComponent);
 	AddDependencyHandle(InComponent->ModuleHandle, InObject, InTickFunction, InEventName, InDependency);
@@ -85,7 +85,7 @@ void UAnimNextComponentWorldSubsystem::AddDependency(UAnimNextComponent* InCompo
 
 void UAnimNextComponentWorldSubsystem::RemoveDependency(UAnimNextComponent* InComponent, UObject* InObject, FTickFunction& InTickFunction, FName InEventName, EDependency InDependency)
 {
-	using namespace UE::AnimNext;
+	using namespace UE::UAF;
 	check(IsInGameThread());
 	check(InComponent);
 	RemoveDependencyHandle(InComponent->ModuleHandle, InObject, InTickFunction, InEventName, InDependency);
@@ -93,7 +93,7 @@ void UAnimNextComponentWorldSubsystem::RemoveDependency(UAnimNextComponent* InCo
 
 void UAnimNextComponentWorldSubsystem::AddModuleEventDependency(UAnimNextComponent* InComponent, FName InEventName, UAnimNextComponent* OtherComponent, FName OtherEventName, EDependency InDependency)
 {
-	using namespace UE::AnimNext;
+	using namespace UE::UAF;
 	check(IsInGameThread());
 	check(InComponent && OtherComponent && InComponent != OtherComponent);
 	AddModuleEventDependencyHandle(InComponent->ModuleHandle, InEventName, OtherComponent->ModuleHandle, OtherEventName, InDependency);
@@ -101,7 +101,7 @@ void UAnimNextComponentWorldSubsystem::AddModuleEventDependency(UAnimNextCompone
 
 void UAnimNextComponentWorldSubsystem::RemoveModuleEventDependency(UAnimNextComponent* InComponent, FName InEventName, UAnimNextComponent* OtherComponent, FName OtherEventName, EDependency InDependency)
 {
-	using namespace UE::AnimNext;
+	using namespace UE::UAF;
 	check(IsInGameThread());
 	check(InComponent && OtherComponent && InComponent != OtherComponent);
 	RemoveModuleEventDependencyHandle(InComponent->ModuleHandle, InEventName, OtherComponent->ModuleHandle, OtherEventName, InDependency);
@@ -109,14 +109,14 @@ void UAnimNextComponentWorldSubsystem::RemoveModuleEventDependency(UAnimNextComp
 
 EPropertyBagResult UAnimNextComponentWorldSubsystem::SetVariable(UAnimNextComponent* InComponent, const FAnimNextVariableReference& InVariable, const FAnimNextParamType& InType, TConstArrayView<uint8> InData)
 {
-	using namespace UE::AnimNext;
+	using namespace UE::UAF;
 	check(InComponent);
 	return SetVariableHandle(InComponent->ModuleHandle, InVariable, InType, InData);
 }
 
 EPropertyBagResult UAnimNextComponentWorldSubsystem::WriteVariable(UAnimNextComponent* InComponent, const FAnimNextVariableReference& InVariable, const FAnimNextParamType& InType, TFunctionRef<void(TArrayView<uint8>)> InFunction)
 {
-	using namespace UE::AnimNext;
+	using namespace UE::UAF;
 	check(InComponent);
 	return WriteVariableHandle(InComponent->ModuleHandle, InVariable, InType, InFunction);
 }
\ No newline at end of file
diff --git a/Engine/Plugins/Experimental/AnimNext/Source/AnimNext/Private/Component/AnimNextComponentWorldSubsystem.h b/Engine/Plugins/Experimental/AnimNext/Source/AnimNext/Private/Component/AnimNextComponentWorldSubsystem.h
--- a/Engine/Plugins/Experimental/AnimNext/Source/AnimNext/Private/Component/AnimNextComponentWorldSubsystem.h
+++ b/Engine/Plugins/Experimental/AnimNext/Source/AnimNext/Private/Component/AnimNextComponentWorldSubsystem.h
@@ -44,7 +44,7 @@ class UAnimNextComponentWorldSubsystem : public UAnimNextWorldSubsystem
 	// @param	InModuleEventName	The name of the event in the module to run the supplied task relative to. If this is NAME_None, then the first valid event will be used.
 	// @param	InTaskFunction		The function to run
 	// @param	InLocation			Where to run the task, before or after
-	void QueueTask(UAnimNextComponent* InComponent, FName InModuleEventName, TUniqueFunction<void(const UE::AnimNext::FModuleTaskContext&)>&& InTaskFunction, UE::AnimNext::ETaskRunLocation InLocation = UE::AnimNext::ETaskRunLocation::Before);
+	void QueueTask(UAnimNextComponent* InComponent, FName InModuleEventName, TUniqueFunction<void(const UE::UAF::FModuleTaskContext&)>&& InTaskFunction, UE::UAF::ETaskRunLocation InLocation = UE::UAF::ETaskRunLocation::Before);
 
 	// Queue an input trait event
 	void QueueInputTraitEvent(UAnimNextComponent* InComponent, FAnimNextTraitEventPtr Event);
diff --git a/Engine/Plugins/Experimental/AnimNext/Source/AnimNext/Private/Component/AnimNextWorldSubsystem.cpp b/Engine/Plugins/Experimental/AnimNext/Source/AnimNext/Private/Component/AnimNextWorldSubsystem.cpp
--- a/Engine/Plugins/Experimental/AnimNext/Source/AnimNext/Private/Component/AnimNextWorldSubsystem.cpp
+++ b/Engine/Plugins/Experimental/AnimNext/Source/AnimNext/Private/Component/AnimNextWorldSubsystem.cpp
@@ -50,7 +50,7 @@ void UAnimNextWorldSubsystem::BeginDestroy()
 				Instance.RemoveAllTickDependencies();
 				Instance.ReleaseComponents();
 			}
-			Instances = UE::AnimNext::TPool<FAnimNextModuleInstance>(); // Force instances destruction
+			Instances = UE::UAF::TPool<FAnimNextModuleInstance>(); // Force instances destruction
 		}
 	}
 }
@@ -97,14 +97,14 @@ bool UAnimNextWorldSubsystem::DoesSupportWorldType(const EWorldType::Type WorldT
 	return false;
 }
 
-bool UAnimNextWorldSubsystem::IsValidHandle(UE::AnimNext::FModuleHandle InHandle) const
+bool UAnimNextWorldSubsystem::IsValidHandle(UE::UAF::FModuleHandle InHandle) const
 {
 	return Instances.IsValidHandle(InHandle);
 }
 
 void UAnimNextWorldSubsystem::FlushPendingActions()
 {
-	using namespace UE::AnimNext;
+	using namespace UE::UAF;
 
 	FRWScopeLock PendingLockScope(PendingLock, SLT_Write);
 
@@ -148,9 +148,9 @@ void UAnimNextWorldSubsystem::FlushPendingActions()
 	}
 }
 
-void UAnimNextWorldSubsystem::RegisterHandle(UE::AnimNext::FModuleHandle& InOutHandle, UAnimNextModule* InModule, UObject* InObject, EAnimNextModuleInitMethod InitMethod)
+void UAnimNextWorldSubsystem::RegisterHandle(UE::UAF::FModuleHandle& InOutHandle, UAnimNextModule* InModule, UObject* InObject, EAnimNextModuleInitMethod InitMethod)
 {
-	using namespace UE::AnimNext;
+	using namespace UE::UAF;
 	check(IsInGameThread());
 	FWriteScopeLock InstancesLockScope(InstancesLock);
 
@@ -160,9 +160,9 @@ void UAnimNextWorldSubsystem::RegisterHandle(UE::AnimNext::FModuleHandle& InOutH
 	Instance.Initialize();
 }
 
-void UAnimNextWorldSubsystem::UnregisterHandle(UE::AnimNext::FModuleHandle& InOutHandle)
+void UAnimNextWorldSubsystem::UnregisterHandle(UE::UAF::FModuleHandle& InOutHandle)
 {
-	using namespace UE::AnimNext;
+	using namespace UE::UAF;
 	check(IsInGameThread());
 
 	if(IsValidHandle(InOutHandle))
@@ -183,9 +183,9 @@ void UAnimNextWorldSubsystem::UnregisterHandle(UE::AnimNext::FModuleHandle& InOu
 	}
 }
 
-bool UAnimNextWorldSubsystem::IsHandleEnabled(UE::AnimNext::FModuleHandle InHandle) const
+bool UAnimNextWorldSubsystem::IsHandleEnabled(UE::UAF::FModuleHandle InHandle) const
 {
-	using namespace UE::AnimNext;
+	using namespace UE::UAF;
 	check(IsInGameThread());
 	if (IsValidHandle(InHandle))
 	{
@@ -204,9 +204,9 @@ bool UAnimNextWorldSubsystem::IsHandleEnabled(UE::AnimNext::FModuleHandle InHand
 	return false;
 }
 
-void UAnimNextWorldSubsystem::EnableHandle(UE::AnimNext::FModuleHandle InHandle, bool bInEnabled)
+void UAnimNextWorldSubsystem::EnableHandle(UE::UAF::FModuleHandle InHandle, bool bInEnabled)
 {
-	using namespace UE::AnimNext;
+	using namespace UE::UAF;
 	check(IsInGameThread());
 	if (IsValidHandle(InHandle))
 	{
@@ -215,9 +215,9 @@ void UAnimNextWorldSubsystem::EnableHandle(UE::AnimNext::FModuleHandle InHandle,
 }
 
 #if UE_ENABLE_DEBUG_DRAWING
-void UAnimNextWorldSubsystem::ShowDebugDrawingHandle(UE::AnimNext::FModuleHandle InHandle, bool bInShowDebugDrawing)
+void UAnimNextWorldSubsystem::ShowDebugDrawingHandle(UE::UAF::FModuleHandle InHandle, bool bInShowDebugDrawing)
 {
-	using namespace UE::AnimNext;
+	using namespace UE::UAF;
 	check(IsInGameThread());
 	if (IsValidHandle(InHandle))
 	{
@@ -226,9 +226,9 @@ void UAnimNextWorldSubsystem::ShowDebugDrawingHandle(UE::AnimNext::FModuleHandle
 }
 #endif
 
-void UAnimNextWorldSubsystem::QueueTaskHandle(UE::AnimNext::FModuleHandle InOutHandle, FName InModuleEventName, TUniqueFunction<void(const UE::AnimNext::FModuleTaskContext&)>&& InTaskFunction, UE::AnimNext::ETaskRunLocation InLocation)
+void UAnimNextWorldSubsystem::QueueTaskHandle(UE::UAF::FModuleHandle InOutHandle, FName InModuleEventName, TUniqueFunction<void(const UE::UAF::FModuleTaskContext&)>&& InTaskFunction, UE::UAF::ETaskRunLocation InLocation)
 {
-	using namespace UE::AnimNext;
+	using namespace UE::UAF;
 	check(IsInGameThread());
 	if (IsValidHandle(InOutHandle))
 	{
@@ -255,9 +255,9 @@ void UAnimNextWorldSubsystem::OnCompileJobFinished(UAnimNextRigVMAsset* InAsset)
 
 #endif
 
-void UAnimNextWorldSubsystem::QueueInputTraitEventHandle(UE::AnimNext::FModuleHandle InHandle, FAnimNextTraitEventPtr Event)
+void UAnimNextWorldSubsystem::QueueInputTraitEventHandle(UE::UAF::FModuleHandle InHandle, FAnimNextTraitEventPtr Event)
 {
-	using namespace UE::AnimNext;
+	using namespace UE::UAF;
 
 	QueueTaskHandle(InHandle, NAME_None, [Event = MoveTemp(Event)](const FModuleTaskContext& InContext)
 	{
@@ -266,9 +266,9 @@ void UAnimNextWorldSubsystem::QueueInputTraitEventHandle(UE::AnimNext::FModuleHa
 	ETaskRunLocation::Before);
 }
 
-const FTickFunction* UAnimNextWorldSubsystem::FindTickFunctionHandle(UE::AnimNext::FModuleHandle InHandle, FName InEventName) const
+const FTickFunction* UAnimNextWorldSubsystem::FindTickFunctionHandle(UE::UAF::FModuleHandle InHandle, FName InEventName) const
 {
-	using namespace UE::AnimNext;
+	using namespace UE::UAF;
 	check(IsInGameThread());
 	if (IsValidHandle(InHandle))
 	{
@@ -290,9 +290,9 @@ const FTickFunction* UAnimNextWorldSubsystem::FindTickFunctionHandle(UE::AnimNex
 	return nullptr;
 }
 
-void UAnimNextWorldSubsystem::AddDependencyHandle(UE::AnimNext::FModuleHandle InHandle, UObject* InObject, FTickFunction& InTickFunction, FName InEventName, EDependency InDependency)
+void UAnimNextWorldSubsystem::AddDependencyHandle(UE::UAF::FModuleHandle InHandle, UObject* InObject, FTickFunction& InTickFunction, FName InEventName, EDependency InDependency)
 {
-	using namespace UE::AnimNext;
+	using namespace UE::UAF;
 	check(IsInGameThread());
 	if (IsValidHandle(InHandle))
 	{
@@ -315,9 +315,9 @@ void UAnimNextWorldSubsystem::AddDependencyHandle(UE::AnimNext::FModuleHandle In
 	}
 }
 
-void UAnimNextWorldSubsystem::RemoveDependencyHandle(UE::AnimNext::FModuleHandle InHandle, UObject* InObject, FTickFunction& InTickFunction, FName InEventName, EDependency InDependency)
+void UAnimNextWorldSubsystem::RemoveDependencyHandle(UE::UAF::FModuleHandle InHandle, UObject* InObject, FTickFunction& InTickFunction, FName InEventName, EDependency InDependency)
 {
-	using namespace UE::AnimNext;
+	using namespace UE::UAF;
 	check(IsInGameThread());
 	if (IsValidHandle(InHandle))
 	{
@@ -340,9 +340,9 @@ void UAnimNextWorldSubsystem::RemoveDependencyHandle(UE::AnimNext::FModuleHandle
 	}
 }
 
-void UAnimNextWorldSubsystem::AddModuleEventDependencyHandle(UE::AnimNext::FModuleHandle InHandle, FName InEventName, UE::AnimNext::FModuleHandle OtherHandle, FName OtherEventName, EDependency InDependency)
+void UAnimNextWorldSubsystem::AddModuleEventDependencyHandle(UE::UAF::FModuleHandle InHandle, FName InEventName, UE::UAF::FModuleHandle OtherHandle, FName OtherEventName, EDependency InDependency)
 {
-	using namespace UE::AnimNext;
+	using namespace UE::UAF;
 	check(IsInGameThread());
 	check(InHandle != OtherHandle);
 	if (IsValidHandle(OtherHandle))
@@ -359,9 +359,9 @@ void UAnimNextWorldSubsystem::AddModuleEventDependencyHandle(UE::AnimNext::FModu
 	}
 }
 
-void UAnimNextWorldSubsystem::RemoveModuleEventDependencyHandle(UE::AnimNext::FModuleHandle InHandle, FName InEventName, UE::AnimNext::FModuleHandle OtherHandle, FName OtherEventName, EDependency InDependency)
+void UAnimNextWorldSubsystem::RemoveModuleEventDependencyHandle(UE::UAF::FModuleHandle InHandle, FName InEventName, UE::UAF::FModuleHandle OtherHandle, FName OtherEventName, EDependency InDependency)
 {
-	using namespace UE::AnimNext;
+	using namespace UE::UAF;
 	check(IsInGameThread());
 	check(InHandle != OtherHandle);
 	if (IsValidHandle(InHandle))
@@ -378,9 +378,9 @@ void UAnimNextWorldSubsystem::RemoveModuleEventDependencyHandle(UE::AnimNext::FM
 	}
 }
 
-EPropertyBagResult UAnimNextWorldSubsystem::SetVariableHandle(UE::AnimNext::FModuleHandle InHandle, const FAnimNextVariableReference& InVariable, const FAnimNextParamType& InType, TConstArrayView<uint8> InData)
+EPropertyBagResult UAnimNextWorldSubsystem::SetVariableHandle(UE::UAF::FModuleHandle InHandle, const FAnimNextVariableReference& InVariable, const FAnimNextParamType& InType, TConstArrayView<uint8> InData)
 {
-	using namespace UE::AnimNext;
+	using namespace UE::UAF;
 	if (IsValidHandle(InHandle))
 	{
 		FAnimNextModuleInstance& Instance = Instances.Get(InHandle);
@@ -389,9 +389,9 @@ EPropertyBagResult UAnimNextWorldSubsystem::SetVariableHandle(UE::AnimNext::FMod
 	return EPropertyBagResult::PropertyNotFound;
 }
 
-EPropertyBagResult UAnimNextWorldSubsystem::WriteVariableHandle(UE::AnimNext::FModuleHandle InHandle, const FAnimNextVariableReference& InVariable, const FAnimNextParamType& InType, TFunctionRef<void(TArrayView<uint8>)> InFunction)
+EPropertyBagResult UAnimNextWorldSubsystem::WriteVariableHandle(UE::UAF::FModuleHandle InHandle, const FAnimNextVariableReference& InVariable, const FAnimNextParamType& InType, TFunctionRef<void(TArrayView<uint8>)> InFunction)
 {
-	using namespace UE::AnimNext;
+	using namespace UE::UAF;
 	if (IsValidHandle(InHandle))
 	{
 		FAnimNextModuleInstance& Instance = Instances.Get(InHandle);
diff --git a/Engine/Plugins/Experimental/AnimNext/Source/AnimNext/Private/Component/SkinnedMeshComponentExtensions.cpp b/Engine/Plugins/Experimental/AnimNext/Source/AnimNext/Private/Component/SkinnedMeshComponentExtensions.cpp
--- a/Engine/Plugins/Experimental/AnimNext/Source/AnimNext/Private/Component/SkinnedMeshComponentExtensions.cpp
+++ b/Engine/Plugins/Experimental/AnimNext/Source/AnimNext/Private/Component/SkinnedMeshComponentExtensions.cpp
@@ -20,12 +20,12 @@ namespace UE::Anim
 
 void ConvertLocalSpaceToComponentSpace(
 	const USkinnedMeshComponent* InComponent,
-	const UE::AnimNext::FLODPoseHeap& InLODPose,
+	const UE::UAF::FLODPoseHeap& InLODPose,
 	TArrayView<FTransform> OutComponentSpaceTransforms)
 {
 	SCOPE_CYCLE_COUNTER(STAT_AnimNext_ConvertLocalSpaceToComponentSpace);
 
-	const UE::AnimNext::FReferencePose& RefPose = InLODPose.GetRefPose();
+	const UE::UAF::FReferencePose& RefPose = InLODPose.GetRefPose();
 
 	const TArrayView<const FBoneIndexType> MeshBoneIndexToParentMeshBoneIndexMap = RefPose.GetMeshBoneIndexToParentMeshBoneIndexMap();
 	const TArrayView<const FBoneIndexType> LODBoneIndexToMeshBoneIndexMap = RefPose.GetLODBoneIndexToMeshBoneIndexMap(InLODPose.LODLevel);
@@ -98,7 +98,7 @@ void ConvertLocalSpaceToComponentSpace(
 
 void FSkinnedMeshComponentExtensions::CompleteAndDispatch(
 	USkinnedMeshComponent* InComponent,
-	const UE::AnimNext::FLODPoseHeap& InLODPose,
+	const UE::UAF::FLODPoseHeap& InLODPose,
 	TUniqueFunction<void(USkinnedMeshComponent*)>&& InGameThreadCallback)
 {
 	bool bNeedsChildTransformUpdate = false;
diff --git a/Engine/Plugins/Experimental/AnimNext/Source/AnimNext/Private/DataRegistry.cpp b/Engine/Plugins/Experimental/AnimNext/Source/AnimNext/Private/DataRegistry.cpp
--- a/Engine/Plugins/Experimental/AnimNext/Source/AnimNext/Private/DataRegistry.cpp
+++ b/Engine/Plugins/Experimental/AnimNext/Source/AnimNext/Private/DataRegistry.cpp
@@ -14,7 +14,7 @@
 namespace // Private
 {
 
-UE::AnimNext::FDataRegistry* GAnimationDataRegistry = nullptr;
+UE::UAF::FDataRegistry* GAnimationDataRegistry = nullptr;
 constexpr int32 BASIC_TYPE_ALLOC_BLOCK = 1000;
 FDelegateHandle PostGarbageCollectHandle;
 
@@ -22,7 +22,7 @@ static TAutoConsoleVariable<float> CVar_GCCleanupTimeBudget(TEXT("a.AnimNext.GCC
 
 } // end namespace
 
-namespace UE::AnimNext
+namespace UE::UAF
 {
 
 /*static*/ void FDataRegistry::Init()
@@ -314,4 +314,4 @@ void FDataRegistry::ReleaseReferencePoseData()
 	SkeletalMeshReferencePoses.Empty();
 }
 
-} // end namespace UE::AnimNext
+} // end namespace UE::UAF
diff --git a/Engine/Plugins/Experimental/AnimNext/Source/AnimNext/Private/DataRegistryTypes.cpp b/Engine/Plugins/Experimental/AnimNext/Source/AnimNext/Private/DataRegistryTypes.cpp
--- a/Engine/Plugins/Experimental/AnimNext/Source/AnimNext/Private/DataRegistryTypes.cpp
+++ b/Engine/Plugins/Experimental/AnimNext/Source/AnimNext/Private/DataRegistryTypes.cpp
@@ -3,7 +3,7 @@
 #include "DataRegistryTypes.h"
 #include "DataRegistry.h"
 
-namespace UE::AnimNext
+namespace UE::UAF
 {
 
 FDataHandle::~FDataHandle()
@@ -22,4 +22,4 @@ FDataHandle::~FDataHandle()
 	}
 }
 
-} // namespace UE::AnimNext
+} // namespace UE::UAF
diff --git a/Engine/Plugins/Experimental/AnimNext/Source/AnimNext/Private/DecompressionTools.cpp b/Engine/Plugins/Experimental/AnimNext/Source/AnimNext/Private/DecompressionTools.cpp
--- a/Engine/Plugins/Experimental/AnimNext/Source/AnimNext/Private/DecompressionTools.cpp
+++ b/Engine/Plugins/Experimental/AnimNext/Source/AnimNext/Private/DecompressionTools.cpp
@@ -22,7 +22,7 @@ DECLARE_CYCLE_STAT(TEXT("Extract Pose From Anim Data"), STAT_ExtractPoseFromAnim
 DECLARE_CYCLE_STAT(TEXT("AnimSeq GetBonePose"), STAT_AnimSeq_GetBonePose, STATGROUP_Anim);
 DECLARE_CYCLE_STAT(TEXT("AnimSeq EvalCurveData"), STAT_AnimSeq_EvalCurveData, STATGROUP_Anim);
 
-namespace UE::AnimNext
+namespace UE::UAF
 {
 
 static IConsoleVariable* CVarForceEvalRawData = IConsoleManager::Get().FindConsoleVariable(TEXT("a.ForceEvalRawData"));
@@ -800,4 +800,4 @@ void FDecompressionTools::RetargetBoneTransform(const FReferencePose& ReferenceP
 	}
 }
 
-} // end namespace UE::AnimNext
+} // end namespace UE::UAF
diff --git a/Engine/Plugins/Experimental/AnimNext/Source/AnimNext/Private/GenerationTools.cpp b/Engine/Plugins/Experimental/AnimNext/Source/AnimNext/Private/GenerationTools.cpp
--- a/Engine/Plugins/Experimental/AnimNext/Source/AnimNext/Private/GenerationTools.cpp
+++ b/Engine/Plugins/Experimental/AnimNext/Source/AnimNext/Private/GenerationTools.cpp
@@ -29,7 +29,7 @@ DEFINE_STAT(STAT_AnimNext_ConvertLocalSpaceToComponentSpace);
 DEFINE_LOG_CATEGORY_STATIC(LogAnimGenerationTools, Log, All)
 
 
-namespace UE::AnimNext
+namespace UE::UAF
 {
 
 struct FCompareBoneIndexType
@@ -47,7 +47,7 @@ struct FCompareBoneIndexType
 {
 	SCOPE_CYCLE_COUNTER(STAT_AnimNext_GenerateReferencePose);
 	
-	using namespace UE::AnimNext;
+	using namespace UE::UAF;
 
 	bool ReferencePoseGenerated = false;
 
@@ -941,11 +941,11 @@ void FGenerationTools::ConvertLocalSpaceToComponentSpace(
 	}
 }
 
-} // end namespace UE::AnimNext
+} // end namespace UE::UAF
 
 
 
-namespace AnimNext::Tools::ConsoleCommands
+namespace UE::UAF::Tools::ConsoleCommands
 {
 	struct FHelper
 	{
@@ -953,14 +953,14 @@ namespace AnimNext::Tools::ConsoleCommands
 		{
 			TArray<FBoneIndexType> ExcludedBones;
 
-			UE::AnimNext::FGenerationTools::DifferenceBoneIndexArrays(LODRequiredBones, NextLODRequiredBones, ExcludedBones);
+			UE::UAF::FGenerationTools::DifferenceBoneIndexArrays(LODRequiredBones, NextLODRequiredBones, ExcludedBones);
 
 			return ExcludedBones;
 		}
 
 		static void CheckSkeletalMeshesLODs()
 		{
-			using namespace UE::AnimNext;
+			using namespace UE::UAF;
 
 			TArray<FAssetData> Assets;
 			IAssetRegistry& AssetRegistry = FModuleManager::LoadModuleChecked<FAssetRegistryModule>(TEXT("AssetRegistry")).Get();
@@ -972,7 +972,7 @@ namespace AnimNext::Tools::ConsoleCommands
 				USkeletalMesh* SkeletalMesh = Cast<USkeletalMesh>(Assets[Idx].GetAsset());
 				if (SkeletalMesh != nullptr)
 				{
-					UE::AnimNext::FReferencePose OutAnimationReferencePose;
+					UE::UAF::FReferencePose OutAnimationReferencePose;
 
 					if (FGenerationTools::GenerateReferencePose(nullptr, SkeletalMesh, OutAnimationReferencePose))
 					{
@@ -992,4 +992,4 @@ namespace AnimNext::Tools::ConsoleCommands
 	{
 		FHelper::CheckSkeletalMeshesLODs();
 	}));
-} // end namespace AnimNext::Tools::ConsoleCommands
+} // end namespace UE::UAF::Tools::ConsoleCommands
diff --git a/Engine/Plugins/Experimental/AnimNext/Source/AnimNext/Private/Graph/RemapPoseData.cpp b/Engine/Plugins/Experimental/AnimNext/Source/AnimNext/Private/Graph/RemapPoseData.cpp
--- a/Engine/Plugins/Experimental/AnimNext/Source/AnimNext/Private/Graph/RemapPoseData.cpp
+++ b/Engine/Plugins/Experimental/AnimNext/Source/AnimNext/Private/Graph/RemapPoseData.cpp
@@ -7,8 +7,8 @@
 
 DEFINE_STAT(STAT_AnimNext_RemapPose_Mesh2Mesh);
 
-bool FRemapPoseData::ShouldReinit(const UE::AnimNext::FReferencePose& InSourceRefPose,
-	const UE::AnimNext::FReferencePose& InTargetRefPose) const
+bool FRemapPoseData::ShouldReinit(const UE::UAF::FReferencePose& InSourceRefPose,
+	const UE::UAF::FReferencePose& InTargetRefPose) const
 {
 	if (!SourceRefPose || !TargetRefPose)
 	{
@@ -31,8 +31,8 @@ bool FRemapPoseData::ShouldReinit(const UE::AnimNext::FReferencePose& InSourceRe
 	return false;
 }
 
-void FRemapPoseData::Reinit(const UE::AnimNext::FReferencePose& InSourceRefPose,
-	const UE::AnimNext::FReferencePose& InTargetRefPose)
+void FRemapPoseData::Reinit(const UE::UAF::FReferencePose& InSourceRefPose,
+	const UE::UAF::FReferencePose& InTargetRefPose)
 {
 	SourceRefPose = &InSourceRefPose;
 	TargetRefPose = &InTargetRefPose;
@@ -93,7 +93,7 @@ void FRemapPoseData::Reinit(const UE::AnimNext::FReferencePose& InSourceRefPose,
 				const int32 SourceBoneIndex = SourceToTargetBoneIndexMap.BoneIndexMap[MappingElementIndex].SourceBoneIndex;
 				const int32 TargetBoneIndex = SourceToTargetBoneIndexMap.BoneIndexMap[MappingElementIndex].TargetBoneIndex;
 
-				if (TargetBoneIndex == UE::AnimNext::FReferencePose::RootBoneIndex)
+				if (TargetBoneIndex == UE::UAF::FReferencePose::RootBoneIndex)
 				{
 					SourceToTargetBoneIndexMap.TargetRootToSourceBoneIndex = SourceBoneIndex;
 					break;
@@ -103,7 +103,7 @@ void FRemapPoseData::Reinit(const UE::AnimNext::FReferencePose& InSourceRefPose,
 	}
 }
 
-FTransform FRemapPoseData::RecursiveCalcModelspaceTransform(const UE::AnimNext::FLODPoseHeap& Pose, FBoneIndexType BoneIndex) const
+FTransform FRemapPoseData::RecursiveCalcModelspaceTransform(const UE::UAF::FLODPoseHeap& Pose, FBoneIndexType BoneIndex) const
 {
 	if (BoneIndex == (uint16)INDEX_NONE)
 	{
@@ -124,8 +124,8 @@ FTransform FRemapPoseData::RecursiveCalcModelspaceTransform(const UE::AnimNext::
 }
 
 // TODO: Source pose dictates the LOD level of the target at the moment. What if the LOD levels of source and target are not in sync? We probably need to use the target skel mesh component LOD level when preparing the target pose.
-void FRemapPoseData::RemapPose(const UE::AnimNext::FLODPoseHeap& SourcePose,
-	UE::AnimNext::FLODPoseHeap& OutTargetPose) const
+void FRemapPoseData::RemapPose(const UE::UAF::FLODPoseHeap& SourcePose,
+	UE::UAF::FLODPoseHeap& OutTargetPose) const
 {
 	SCOPE_CYCLE_COUNTER(STAT_AnimNext_RemapPose_Mesh2Mesh);
 
@@ -195,13 +195,13 @@ void FRemapPoseData::RemapPose(const UE::AnimNext::FLODPoseHeap& SourcePose,
 	{
 		// TODO: Replace with the generic method once we have access to the modelspace transform via the LOD pose.
 		const FTransform RootModelTransform = RecursiveCalcModelspaceTransform(SourcePose, RootBoneOnSourceMapped);
-		OutTargetPose.LocalTransforms[UE::AnimNext::FReferencePose::RootBoneIndex] = RootModelTransform;
+		OutTargetPose.LocalTransforms[UE::UAF::FReferencePose::RootBoneIndex] = RootModelTransform;
 	}
 }
 
-void FRemapPoseData::RemapAttributes(const UE::AnimNext::FLODPose& SourceLODPose,
+void FRemapPoseData::RemapAttributes(const UE::UAF::FLODPose& SourceLODPose,
 	const UE::Anim::FHeapAttributeContainer& InAttributes,
-	const UE::AnimNext::FLODPose& TargetLODPose,
+	const UE::UAF::FLODPose& TargetLODPose,
 	UE::Anim::FHeapAttributeContainer& OutAttributes)
 {
 	const USkeletalMesh* SourceSkeletalMesh = SourceRefPose->SkeletalMesh.Get();
diff --git a/Engine/Plugins/Experimental/AnimNext/Source/AnimNext/Private/Graph/RemapPoseData.h b/Engine/Plugins/Experimental/AnimNext/Source/AnimNext/Private/Graph/RemapPoseData.h
--- a/Engine/Plugins/Experimental/AnimNext/Source/AnimNext/Private/Graph/RemapPoseData.h
+++ b/Engine/Plugins/Experimental/AnimNext/Source/AnimNext/Private/Graph/RemapPoseData.h
@@ -45,34 +45,34 @@ struct FRemapPoseData
 	 * Check if we need to re-initialize the mapping. This will return true in case of first time usage of the data,
 	 * or in the case that either the source or target pose changed.
 	 */
-	bool ShouldReinit(const UE::AnimNext::FReferencePose& InSourceRefPose,
-		const UE::AnimNext::FReferencePose& InTargetRefPose) const;
+	bool ShouldReinit(const UE::UAF::FReferencePose& InSourceRefPose,
+		const UE::UAF::FReferencePose& InTargetRefPose) const;
 
 	/**
 	 * Re-initializes the mapping for each LOD. Call ShouldReinit() before to avoid unnecessary compute.
 	 * Compare the bones from the source and target pose and find the subset of the target that is also present in the source.
 	 */
-	void Reinit(const UE::AnimNext::FReferencePose& InSourceRefPose,
-		const UE::AnimNext::FReferencePose& InTargetRefPose);
+	void Reinit(const UE::UAF::FReferencePose& InSourceRefPose,
+		const UE::UAF::FReferencePose& InTargetRefPose);
 
 	/**
 	 * Remap a pose using the previously cached mapping. Call this at runtime.
 	 */
-	void RemapPose(const UE::AnimNext::FLODPoseHeap& SourcePose,
-		UE::AnimNext::FLODPoseHeap& OutTargetPose) const;
+	void RemapPose(const UE::UAF::FLODPoseHeap& SourcePose,
+		UE::UAF::FLODPoseHeap& OutTargetPose) const;
 
-	void RemapAttributes(const UE::AnimNext::FLODPose& SourceLODPose,
+	void RemapAttributes(const UE::UAF::FLODPose& SourceLODPose,
 		const UE::Anim::FHeapAttributeContainer& InAttributes,
-		const UE::AnimNext::FLODPose& TargetLODPose,
+		const UE::UAF::FLODPose& TargetLODPose,
 		UE::Anim::FHeapAttributeContainer& OutAttributes);
 
 	/**
 	 * Temporary calculate modelspace transform, until the LODPose is able to do that.
 	 */
-	FTransform RecursiveCalcModelspaceTransform(const UE::AnimNext::FLODPoseHeap& Pose, FBoneIndexType BoneIndex) const;
+	FTransform RecursiveCalcModelspaceTransform(const UE::UAF::FLODPoseHeap& Pose, FBoneIndexType BoneIndex) const;
 
-	const UE::AnimNext::FReferencePose* SourceRefPose = nullptr;
-	const UE::AnimNext::FReferencePose* TargetRefPose = nullptr;
+	const UE::UAF::FReferencePose* SourceRefPose = nullptr;
+	const UE::UAF::FReferencePose* TargetRefPose = nullptr;
 
 	/**
 	 * Bone mapping from the source to the target skeletal mesh per LOD level combination.
diff --git a/Engine/Plugins/Experimental/AnimNext/Source/AnimNext/Private/Graph/RigUnit_AnimNextRemapPose.cpp b/Engine/Plugins/Experimental/AnimNext/Source/AnimNext/Private/Graph/RigUnit_AnimNextRemapPose.cpp
--- a/Engine/Plugins/Experimental/AnimNext/Source/AnimNext/Private/Graph/RigUnit_AnimNextRemapPose.cpp
+++ b/Engine/Plugins/Experimental/AnimNext/Source/AnimNext/Private/Graph/RigUnit_AnimNextRemapPose.cpp
@@ -13,7 +13,7 @@ FRigUnit_AnimNextRemapPose_Execute()
 {
 	SCOPE_CYCLE_COUNTER(STAT_AnimNext_RigUnit_RemapPose);
 
-	using namespace UE::AnimNext;
+	using namespace UE::UAF;
 
 	if(!TargetAnimGraphRefPose.ReferencePose.IsValid())
 	{
@@ -27,12 +27,12 @@ FRigUnit_AnimNextRemapPose_Execute()
 
 	FMemMark MemMark(FMemStack::Get());
 
-	const UE::AnimNext::FLODPoseHeap& SourcePose = Pose.LODPose;
-	const UE::AnimNext::FReferencePose& SourceRefPose = SourcePose.GetRefPose();
+	const UE::UAF::FLODPoseHeap& SourcePose = Pose.LODPose;
+	const UE::UAF::FReferencePose& SourceRefPose = SourcePose.GetRefPose();
 	const USkeletalMesh* SourceMesh = SourceRefPose.SkeletalMesh.Get();
 
-	const UE::AnimNext::FReferencePose& TargetRefPose = TargetAnimGraphRefPose.ReferencePose.GetRef<UE::AnimNext::FReferencePose>();
-	UE::AnimNext::FLODPoseHeap& TargetPose = Result.LODPose;
+	const UE::UAF::FReferencePose& TargetRefPose = TargetAnimGraphRefPose.ReferencePose.GetRef<UE::UAF::FReferencePose>();
+	UE::UAF::FLODPoseHeap& TargetPose = Result.LODPose;
 	const USkeletalMesh* TargetMesh = TargetRefPose.SkeletalMesh.Get();
 
 	if (SourceMesh == TargetMesh)
diff --git a/Engine/Plugins/Experimental/AnimNext/Source/AnimNext/Private/Graph/RigUnit_AnimNextWriteSkeletalMeshComponentPose.cpp b/Engine/Plugins/Experimental/AnimNext/Source/AnimNext/Private/Graph/RigUnit_AnimNextWriteSkeletalMeshComponentPose.cpp
--- a/Engine/Plugins/Experimental/AnimNext/Source/AnimNext/Private/Graph/RigUnit_AnimNextWriteSkeletalMeshComponentPose.cpp
+++ b/Engine/Plugins/Experimental/AnimNext/Source/AnimNext/Private/Graph/RigUnit_AnimNextWriteSkeletalMeshComponentPose.cpp
@@ -21,7 +21,7 @@ FRigUnit_AnimNextWriteSkeletalMeshComponentPose_Execute()
 {
 	SCOPE_CYCLE_COUNTER(STAT_AnimNext_Write_Pose);
 
-	using namespace UE::AnimNext;
+	using namespace UE::UAF;
 	FTaskTagScope TaskTag(ETaskTag::EParallelGameThread);
 
 	USkeletalMeshComponent* OutputComponent = SkeletalMeshComponent;
@@ -60,7 +60,7 @@ FRigUnit_AnimNextWriteSkeletalMeshComponentPose_Execute()
 	}
 
 	FDataHandle RefPoseHandle = FDataRegistry::Get()->GetOrGenerateReferencePose(OutputComponent);
-	const UE::AnimNext::FReferencePose& RefPose = RefPoseHandle.GetRef<UE::AnimNext::FReferencePose>();
+	const UE::UAF::FReferencePose& RefPose = RefPoseHandle.GetRef<UE::UAF::FReferencePose>();
 
 	// We cannot use two different reference poses because we want to avoid an expensive remap operation
 	// You should remap the pose explicitly if this is what you need
diff --git a/Engine/Plugins/Experimental/AnimNext/Source/AnimNext/Private/Graph/RigUnit_GetActorTransform.cpp b/Engine/Plugins/Experimental/AnimNext/Source/AnimNext/Private/Graph/RigUnit_GetActorTransform.cpp
--- a/Engine/Plugins/Experimental/AnimNext/Source/AnimNext/Private/Graph/RigUnit_GetActorTransform.cpp
+++ b/Engine/Plugins/Experimental/AnimNext/Source/AnimNext/Private/Graph/RigUnit_GetActorTransform.cpp
@@ -39,7 +39,7 @@ FAnimNextActorTransformComponent::FAnimNextActorTransformComponent()
 					return;
 				}
 
-				Component->QueueTask(NAME_None, [Transform = UpdatedComponentOwner->GetActorTransform()](const UE::AnimNext::FModuleTaskContext& InContext)
+				Component->QueueTask(NAME_None, [Transform = UpdatedComponentOwner->GetActorTransform()](const UE::UAF::FModuleTaskContext& InContext)
 				{
 					InContext.TryAccessComponent<FAnimNextActorTransformComponent>([&Transform](FAnimNextActorTransformComponent& InComponent)
 					{
@@ -52,7 +52,7 @@ FAnimNextActorTransformComponent::FAnimNextActorTransformComponent()
 
 FRigUnit_GetActorTransform_Execute()
 {
-	using namespace UE::AnimNext;
+	using namespace UE::UAF;
 
 	const FAnimNextModuleContextData& ModuleContextData = ExecuteContext.GetContextData<FAnimNextModuleContextData>();
 	FAnimNextModuleInstance& ModuleInstance = ModuleContextData.GetModuleInstance();
diff --git a/Engine/Plugins/Experimental/AnimNext/Source/AnimNext/Private/Graph/RigUnit_MakeReferencePoseFromSkeletalMeshComponent.cpp b/Engine/Plugins/Experimental/AnimNext/Source/AnimNext/Private/Graph/RigUnit_MakeReferencePoseFromSkeletalMeshComponent.cpp
--- a/Engine/Plugins/Experimental/AnimNext/Source/AnimNext/Private/Graph/RigUnit_MakeReferencePoseFromSkeletalMeshComponent.cpp
+++ b/Engine/Plugins/Experimental/AnimNext/Source/AnimNext/Private/Graph/RigUnit_MakeReferencePoseFromSkeletalMeshComponent.cpp
@@ -17,7 +17,7 @@ FRigUnit_MakeReferencePoseFromSkeletalMeshComponent_Execute()
 {
 	SCOPE_CYCLE_COUNTER(STAT_AnimNext_Make_RefPose);
 
-	using namespace UE::AnimNext;
+	using namespace UE::UAF;
 
 	USkeletalMeshComponent* InputComponent = SkeletalMeshComponent;
 
diff --git a/Engine/Plugins/Experimental/AnimNext/Source/AnimNext/Private/Module/AnimNextActorComponentReferenceComponent.cpp b/Engine/Plugins/Experimental/AnimNext/Source/AnimNext/Private/Module/AnimNextActorComponentReferenceComponent.cpp
--- a/Engine/Plugins/Experimental/AnimNext/Source/AnimNext/Private/Module/AnimNextActorComponentReferenceComponent.cpp
+++ b/Engine/Plugins/Experimental/AnimNext/Source/AnimNext/Private/Module/AnimNextActorComponentReferenceComponent.cpp
@@ -33,7 +33,7 @@ void FAnimNextActorComponentReferenceComponent::OnInitializeHelper(UScriptStruct
 			}
 
 			UActorComponent* FoundComponent = Owner->FindComponentByClass(ComponentType);
-			PinnedAnimNextComponent->QueueTask(NAME_None, [FoundComponent, InScriptStruct](const UE::AnimNext::FModuleTaskContext& InContext)
+			PinnedAnimNextComponent->QueueTask(NAME_None, [FoundComponent, InScriptStruct](const UE::UAF::FModuleTaskContext& InContext)
 			{
 				InContext.TryAccessComponent(InScriptStruct, [FoundComponent](FUAFModuleInstanceComponent& InComponent)
 				{
diff --git a/Engine/Plugins/Experimental/AnimNext/Source/AnimNext/Private/Module/AnimNextModuleInstance.cpp b/Engine/Plugins/Experimental/AnimNext/Source/AnimNext/Private/Module/AnimNextModuleInstance.cpp
--- a/Engine/Plugins/Experimental/AnimNext/Source/AnimNext/Private/Module/AnimNextModuleInstance.cpp
+++ b/Engine/Plugins/Experimental/AnimNext/Source/AnimNext/Private/Module/AnimNextModuleInstance.cpp
@@ -27,7 +27,7 @@ FAnimNextModuleInstance::FAnimNextModuleInstance() = default;
 FAnimNextModuleInstance::FAnimNextModuleInstance(
 		UAnimNextModule* InModule,
 		UObject* InObject,
-		UE::AnimNext::TPool<FAnimNextModuleInstance>* InPool,
+		UE::UAF::TPool<FAnimNextModuleInstance>* InPool,
 		EAnimNextModuleInitMethod InInitMethod)
 	: Object(InObject)
 	, Pool(InPool)
@@ -43,7 +43,7 @@ FAnimNextModuleInstance::FAnimNextModuleInstance(
 #if UE_ENABLE_DEBUG_DRAWING
 	if(Object && Object->GetWorld())
 	{
-		DebugDraw = MakeUnique<UE::AnimNext::Debug::FDebugDraw>(Object);
+		DebugDraw = MakeUnique<UE::UAF::Debug::FDebugDraw>(Object);
 	}
 #endif
 }
@@ -61,7 +61,7 @@ FAnimNextModuleInstance::~FAnimNextModuleInstance()
 	Handle.Reset();
 }
 
-namespace UE::AnimNext::Private
+namespace UE::UAF::Private
 {
 
 struct FImplementedModuleEvent
@@ -167,7 +167,7 @@ void FAnimNextModuleInstance::Initialize()
 {
 	SCOPE_CYCLE_COUNTER(STAT_AnimNext_InitializeInstance);
 	
-	using namespace UE::AnimNext;
+	using namespace UE::UAF;
 
 	check(IsInGameThread());
 
@@ -328,7 +328,7 @@ void FAnimNextModuleInstance::Initialize()
 
 void FAnimNextModuleInstance::RemoveAllTickDependencies()
 {
-	using namespace UE::AnimNext;
+	using namespace UE::UAF;
 
 	check(IsInGameThread());
 
@@ -340,7 +340,7 @@ void FAnimNextModuleInstance::RemoveAllTickDependencies()
 
 void FAnimNextModuleInstance::ResetBindingsAndInstanceData()
 {
-	using namespace UE::AnimNext;
+	using namespace UE::UAF;
 
 	check(IsInGameThread());
 
@@ -378,7 +378,7 @@ void FAnimNextModuleInstance::QueueOutputTraitEvent(FAnimNextTraitEventPtr Event
 
 bool FAnimNextModuleInstance::IsEnabled() const
 {
-	using namespace UE::AnimNext;
+	using namespace UE::UAF;
 
 	check(IsInGameThread());
 
@@ -387,7 +387,7 @@ bool FAnimNextModuleInstance::IsEnabled() const
 
 void FAnimNextModuleInstance::Enable(bool bInEnabled)
 {
-	using namespace UE::AnimNext;
+	using namespace UE::UAF;
 
 	check(IsInGameThread());
 
@@ -488,7 +488,7 @@ const UAnimNextModule* FAnimNextModuleInstance::GetModule() const
 #if WITH_EDITOR
 void FAnimNextModuleInstance::OnCompileJobFinished()
 {
-	using namespace UE::AnimNext;
+	using namespace UE::UAF;
 
 	FGuardValue_Bitfield(bIsRecreatingOnCompile, true);
 
@@ -514,19 +514,19 @@ void FAnimNextModuleInstance::ShowDebugDrawing(bool bInShowDebugDrawing)
 
 void FAnimNextModuleInstance::RunTaskOnGameThread(TUniqueFunction<void(void)>&& InFunction)
 {
-	UE::AnimNext::FModuleEventTickFunction::RunTaskOnGameThread(MoveTemp(InFunction));
+	UE::UAF::FModuleEventTickFunction::RunTaskOnGameThread(MoveTemp(InFunction));
 }
 
-UE::AnimNext::FModuleEventTickFunction* FAnimNextModuleInstance::FindTickFunctionByName(FName InEventName)
+UE::UAF::FModuleEventTickFunction* FAnimNextModuleInstance::FindTickFunctionByName(FName InEventName)
 {
-	using namespace UE::AnimNext;
+	using namespace UE::UAF;
 
 	return const_cast<FModuleEventTickFunction*>(const_cast<const FAnimNextModuleInstance*>(this)->FindTickFunctionByName(InEventName));
 }
 
-const UE::AnimNext::FModuleEventTickFunction* FAnimNextModuleInstance::FindTickFunctionByName(FName InEventName) const
+const UE::UAF::FModuleEventTickFunction* FAnimNextModuleInstance::FindTickFunctionByName(FName InEventName) const
 {
-	using namespace UE::AnimNext;
+	using namespace UE::UAF;
 
 	for(const FModuleEventTickFunction& TickFunction : TickFunctions)
 	{
@@ -555,7 +555,7 @@ void FAnimNextModuleInstance::EndExecution(float InDeltaTime)
 	}
 }
 
-void FAnimNextModuleInstance::RaiseTraitEvents(const UE::AnimNext::FTraitEventList& EventList)
+void FAnimNextModuleInstance::RaiseTraitEvents(const UE::UAF::FTraitEventList& EventList)
 {
 	for(auto It = ComponentMap.CreateIterator(); It; ++It)
 	{
@@ -628,21 +628,21 @@ void FAnimNextModuleInstance::RunRigVMEvent(FName InEventName, float InDeltaTime
 
 	// Insert our context data for the scope of execution
 	FAnimNextModuleContextData ContextData(this);
-	UE::AnimNext::FScopedExecuteContextData ContextDataScope(AnimNextContext, ContextData);
+	UE::UAF::FScopedExecuteContextData ContextDataScope(AnimNextContext, ContextData);
 
 	// Run the VM for this event
 	VM->ExecuteVM(ExtendedExecuteContext, InEventName);
 }
 
 
-TArrayView<UE::AnimNext::FModuleEventTickFunction> FAnimNextModuleInstance::GetTickFunctions()
+TArrayView<UE::UAF::FModuleEventTickFunction> FAnimNextModuleInstance::GetTickFunctions()
 {
 	return TickFunctions;
 }
 
-UE::AnimNext::FModuleEventTickFunction* FAnimNextModuleInstance::FindFirstUserTickFunction()
+UE::UAF::FModuleEventTickFunction* FAnimNextModuleInstance::FindFirstUserTickFunction()
 {
-	UE::AnimNext::FModuleEventTickFunction* FoundTickFunction = TickFunctions.FindByPredicate([](const UE::AnimNext::FModuleEventTickFunction& InTickFunction)
+	UE::UAF::FModuleEventTickFunction* FoundTickFunction = TickFunctions.FindByPredicate([](const UE::UAF::FModuleEventTickFunction& InTickFunction)
 	{
 		return InTickFunction.bFirstUserEvent;
 	});
@@ -650,9 +650,9 @@ UE::AnimNext::FModuleEventTickFunction* FAnimNextModuleInstance::FindFirstUserTi
 	return FoundTickFunction;
 }
 
-void FAnimNextModuleInstance::QueueTask(FName InEventName, TUniqueFunction<void(const UE::AnimNext::FModuleTaskContext&)>&& InTaskFunction, UE::AnimNext::ETaskRunLocation InLocation)
+void FAnimNextModuleInstance::QueueTask(FName InEventName, TUniqueFunction<void(const UE::UAF::FModuleTaskContext&)>&& InTaskFunction, UE::UAF::ETaskRunLocation InLocation)
 {
-	using namespace UE::AnimNext;
+	using namespace UE::UAF;
 
 	FModuleEventTickFunction* FoundTickFunction = nullptr;
 	if(TickFunctions.Num() > 0)
@@ -669,7 +669,7 @@ void FAnimNextModuleInstance::QueueTask(FName InEventName, TUniqueFunction<void(
 		if (FoundTickFunction == nullptr)
 		{
 			// Fall back to first user function
-			FoundTickFunction = TickFunctions.FindByPredicate([](const UE::AnimNext::FModuleEventTickFunction& InTickFunction)
+			FoundTickFunction = TickFunctions.FindByPredicate([](const UE::UAF::FModuleEventTickFunction& InTickFunction)
 			{
 				return InTickFunction.bFirstUserEvent;
 			});
@@ -700,7 +700,7 @@ void FAnimNextModuleInstance::QueueTask(FName InEventName, TUniqueFunction<void(
 	}
 }
 
-void FAnimNextModuleInstance::QueueTaskOnOtherModule(const UE::AnimNext::FModuleHandle InOtherModuleHandle, const FName InEventName, TUniqueFunction<void(const UE::AnimNext::FModuleTaskContext&)>&& InTaskFunction, UE::AnimNext::ETaskRunLocation InLocation)
+void FAnimNextModuleInstance::QueueTaskOnOtherModule(const UE::UAF::FModuleHandle InOtherModuleHandle, const FName InEventName, TUniqueFunction<void(const UE::UAF::FModuleTaskContext&)>&& InTaskFunction, UE::UAF::ETaskRunLocation InLocation)
 {
 	if (Pool == nullptr)
 	{
diff --git a/Engine/Plugins/Experimental/AnimNext/Source/AnimNext/Private/Module/RigUnit_AnimNextModuleEventDependencies.cpp b/Engine/Plugins/Experimental/AnimNext/Source/AnimNext/Private/Module/RigUnit_AnimNextModuleEventDependencies.cpp
--- a/Engine/Plugins/Experimental/AnimNext/Source/AnimNext/Private/Module/RigUnit_AnimNextModuleEventDependencies.cpp
+++ b/Engine/Plugins/Experimental/AnimNext/Source/AnimNext/Private/Module/RigUnit_AnimNextModuleEventDependencies.cpp
@@ -14,7 +14,7 @@
 #include "Widgets/Layout/SBox.h"
 #endif
 
-namespace UE::AnimNext
+namespace UE::UAF
 {
 
 // Helper function used to add/remove dependencies
@@ -236,7 +236,7 @@ FRigVMTrait_ModuleEventDependency::FRigVMTrait_ModuleEventDependency()
 
 FRigUnit_AnimNextModuleAddDependencies_Execute()
 {
-	UE::AnimNext::ApplyDependenciesHelper(ExecuteContext, [](const UE::AnimNext::FModuleDependencyContext& InContext, const FRigVMTrait_ModuleEventDependency& InDependency)
+	UE::UAF::ApplyDependenciesHelper(ExecuteContext, [](const UE::UAF::FModuleDependencyContext& InContext, const FRigVMTrait_ModuleEventDependency& InDependency)
 	{
 		InDependency.OnAddDependency(InContext);
 	});
@@ -244,7 +244,7 @@ FRigUnit_AnimNextModuleAddDependencies_Execute()
 
 FRigUnit_AnimNextModuleRemoveDependencies_Execute()
 {
-	UE::AnimNext::ApplyDependenciesHelper(ExecuteContext, [](const UE::AnimNext::FModuleDependencyContext& InContext, const FRigVMTrait_ModuleEventDependency& InDependency)
+	UE::UAF::ApplyDependenciesHelper(ExecuteContext, [](const UE::UAF::FModuleDependencyContext& InContext, const FRigVMTrait_ModuleEventDependency& InDependency)
 	{
 		InDependency.OnRemoveDependency(InContext);
 	});
diff --git a/Engine/Plugins/Experimental/AnimNext/Source/AnimNext/Private/Module/RigUnit_AnimNextModuleEvents.cpp b/Engine/Plugins/Experimental/AnimNext/Source/AnimNext/Private/Module/RigUnit_AnimNextModuleEvents.cpp
--- a/Engine/Plugins/Experimental/AnimNext/Source/AnimNext/Private/Module/RigUnit_AnimNextModuleEvents.cpp
+++ b/Engine/Plugins/Experimental/AnimNext/Source/AnimNext/Private/Module/RigUnit_AnimNextModuleEvents.cpp
@@ -6,7 +6,7 @@
 #include "Engine/World.h"
 #include "Module/AnimNextModuleInstance.h"
 
-namespace UE::AnimNext::CVars
+namespace UE::UAF::CVars
 {
 	// Expose control over whether or not the AnimNext binding execution should happen before others in the tick group
 	// This can allow them to dispatch animation jobs earlier, helping hide their latency
@@ -20,11 +20,11 @@ FRigUnit_AnimNextExecuteBindings_GT_Execute()
 {
 }
 
-UE::AnimNext::FModuleEventBindingFunction FRigUnit_AnimNextExecuteBindings_GT::GetBindingFunction() const
+UE::UAF::FModuleEventBindingFunction FRigUnit_AnimNextExecuteBindings_GT::GetBindingFunction() const
 {
-	return [](const UE::AnimNext::FTickFunctionBindingContext& InContext, FTickFunction& InTickFunction)
+	return [](const UE::UAF::FTickFunctionBindingContext& InContext, FTickFunction& InTickFunction)
 	{
-		const bool bHighPriority = UE::AnimNext::CVars::CVarHighPriorityAnimNextExecuteBindingsTick.GetValueOnGameThread();
+		const bool bHighPriority = UE::UAF::CVars::CVarHighPriorityAnimNextExecuteBindingsTick.GetValueOnGameThread();
 		if (InTickFunction.bHighPriority != bHighPriority)
 		{
 			InTickFunction.SetPriorityIncludingPrerequisites(bHighPriority);
@@ -36,11 +36,11 @@ FRigUnit_AnimNextExecuteBindings_WT_Execute()
 {
 }
 
-UE::AnimNext::FModuleEventBindingFunction FRigUnit_AnimNextExecuteBindings_WT::GetBindingFunction() const
+UE::UAF::FModuleEventBindingFunction FRigUnit_AnimNextExecuteBindings_WT::GetBindingFunction() const
 {
-	return [](const UE::AnimNext::FTickFunctionBindingContext& InContext, FTickFunction& InTickFunction)
+	return [](const UE::UAF::FTickFunctionBindingContext& InContext, FTickFunction& InTickFunction)
 	{
-		const bool bHighPriority = UE::AnimNext::CVars::CVarHighPriorityAnimNextExecuteBindingsTick.GetValueOnGameThread();
+		const bool bHighPriority = UE::UAF::CVars::CVarHighPriorityAnimNextExecuteBindingsTick.GetValueOnGameThread();
 		if (InTickFunction.bHighPriority != bHighPriority)
 		{
 			InTickFunction.SetPriorityIncludingPrerequisites(bHighPriority);
@@ -63,9 +63,9 @@ FString FRigUnit_AnimNextUserEvent::GetUnitSubTitle() const
 	return Subtitle;
 }
 
-UE::AnimNext::FModuleEventBindingFunction FRigUnit_AnimNextUserEvent::GetBindingFunction() const
+UE::UAF::FModuleEventBindingFunction FRigUnit_AnimNextUserEvent::GetBindingFunction() const
 {
-	return [TickGroup = TickGroup](const UE::AnimNext::FTickFunctionBindingContext& InContext, FTickFunction& InTickFunction)
+	return [TickGroup = TickGroup](const UE::UAF::FTickFunctionBindingContext& InContext, FTickFunction& InTickFunction)
 	{
 		InTickFunction.TickGroup = TickGroup;
 	};
diff --git a/Engine/Plugins/Experimental/AnimNext/Source/AnimNext/Private/Module/RigVMTrait_ModuleEventDependency_ActorComponentClassPrimaryTickFunction.cpp b/Engine/Plugins/Experimental/AnimNext/Source/AnimNext/Private/Module/RigVMTrait_ModuleEventDependency_ActorComponentClassPrimaryTickFunction.cpp
--- a/Engine/Plugins/Experimental/AnimNext/Source/AnimNext/Private/Module/RigVMTrait_ModuleEventDependency_ActorComponentClassPrimaryTickFunction.cpp
+++ b/Engine/Plugins/Experimental/AnimNext/Source/AnimNext/Private/Module/RigVMTrait_ModuleEventDependency_ActorComponentClassPrimaryTickFunction.cpp
@@ -9,7 +9,7 @@ FString FRigVMTrait_ModuleEventDependency_ActorComponentClassPrimaryTickFunction
 }
 #endif
 
-void FRigVMTrait_ModuleEventDependency_ActorComponentClassPrimaryTickFunction::OnAddDependency(const UE::AnimNext::FModuleDependencyContext& InContext) const
+void FRigVMTrait_ModuleEventDependency_ActorComponentClassPrimaryTickFunction::OnAddDependency(const UE::UAF::FModuleDependencyContext& InContext) const
 {
 	UActorComponent* AnimNextComponent = Cast<UActorComponent>(InContext.Object);
 	if(AnimNextComponent == nullptr)
@@ -39,7 +39,7 @@ void FRigVMTrait_ModuleEventDependency_ActorComponentClassPrimaryTickFunction::O
 	}
 }
 
-void FRigVMTrait_ModuleEventDependency_ActorComponentClassPrimaryTickFunction::OnRemoveDependency(const UE::AnimNext::FModuleDependencyContext& InContext) const
+void FRigVMTrait_ModuleEventDependency_ActorComponentClassPrimaryTickFunction::OnRemoveDependency(const UE::UAF::FModuleDependencyContext& InContext) const
 {
 	UActorComponent* AnimNextComponent = Cast<UActorComponent>(InContext.Object);
 	if(AnimNextComponent == nullptr)
diff --git a/Engine/Plugins/Experimental/AnimNext/Source/AnimNext/Private/Module/RigVMTrait_ModuleEventDependency_ActorComponentClassPrimaryTickFunction.h b/Engine/Plugins/Experimental/AnimNext/Source/AnimNext/Private/Module/RigVMTrait_ModuleEventDependency_ActorComponentClassPrimaryTickFunction.h
--- a/Engine/Plugins/Experimental/AnimNext/Source/AnimNext/Private/Module/RigVMTrait_ModuleEventDependency_ActorComponentClassPrimaryTickFunction.h
+++ b/Engine/Plugins/Experimental/AnimNext/Source/AnimNext/Private/Module/RigVMTrait_ModuleEventDependency_ActorComponentClassPrimaryTickFunction.h
@@ -20,8 +20,8 @@ struct FRigVMTrait_ModuleEventDependency_ActorComponentClassPrimaryTickFunction
 #endif
 
 	// FAnimNextModuleEventDependency interface
-	virtual void OnAddDependency(const UE::AnimNext::FModuleDependencyContext& InContext) const override;
-	virtual void OnRemoveDependency(const UE::AnimNext::FModuleDependencyContext& InContext) const override;
+	virtual void OnAddDependency(const UE::UAF::FModuleDependencyContext& InContext) const override;
+	virtual void OnRemoveDependency(const UE::UAF::FModuleDependencyContext& InContext) const override;
 
 	// The component class to look for when establishing the dependency
 	UPROPERTY(EditAnywhere, Category = "Dependency")
diff --git a/Engine/Plugins/Experimental/AnimNext/Source/AnimNext/Private/Module/RigVMTrait_ModuleEventDependency_ActorComponentPrimaryTickFunction.cpp b/Engine/Plugins/Experimental/AnimNext/Source/AnimNext/Private/Module/RigVMTrait_ModuleEventDependency_ActorComponentPrimaryTickFunction.cpp
--- a/Engine/Plugins/Experimental/AnimNext/Source/AnimNext/Private/Module/RigVMTrait_ModuleEventDependency_ActorComponentPrimaryTickFunction.cpp
+++ b/Engine/Plugins/Experimental/AnimNext/Source/AnimNext/Private/Module/RigVMTrait_ModuleEventDependency_ActorComponentPrimaryTickFunction.cpp
@@ -9,7 +9,7 @@ FString FRigVMTrait_ModuleEventDependency_ActorComponentPrimaryTickFunction::Get
 }
 #endif
 
-void FRigVMTrait_ModuleEventDependency_ActorComponentPrimaryTickFunction::OnAddDependency(const UE::AnimNext::FModuleDependencyContext& InContext) const
+void FRigVMTrait_ModuleEventDependency_ActorComponentPrimaryTickFunction::OnAddDependency(const UE::UAF::FModuleDependencyContext& InContext) const
 {
 	UActorComponent* AnimNextComponent = Cast<UActorComponent>(InContext.Object);
 	if(AnimNextComponent == nullptr)
@@ -32,7 +32,7 @@ void FRigVMTrait_ModuleEventDependency_ActorComponentPrimaryTickFunction::OnAddD
 	}
 }
 
-void FRigVMTrait_ModuleEventDependency_ActorComponentPrimaryTickFunction::OnRemoveDependency(const UE::AnimNext::FModuleDependencyContext& InContext) const
+void FRigVMTrait_ModuleEventDependency_ActorComponentPrimaryTickFunction::OnRemoveDependency(const UE::UAF::FModuleDependencyContext& InContext) const
 {
 	UActorComponent* AnimNextComponent = Cast<UActorComponent>(InContext.Object);
 	if(AnimNextComponent == nullptr)
diff --git a/Engine/Plugins/Experimental/AnimNext/Source/AnimNext/Private/Module/RigVMTrait_ModuleEventDependency_ActorComponentPrimaryTickFunction.h b/Engine/Plugins/Experimental/AnimNext/Source/AnimNext/Private/Module/RigVMTrait_ModuleEventDependency_ActorComponentPrimaryTickFunction.h
--- a/Engine/Plugins/Experimental/AnimNext/Source/AnimNext/Private/Module/RigVMTrait_ModuleEventDependency_ActorComponentPrimaryTickFunction.h
+++ b/Engine/Plugins/Experimental/AnimNext/Source/AnimNext/Private/Module/RigVMTrait_ModuleEventDependency_ActorComponentPrimaryTickFunction.h
@@ -19,8 +19,8 @@ struct FRigVMTrait_ModuleEventDependency_ActorComponentPrimaryTickFunction : pub
 #endif
 
 	// FAnimNextModuleEventDependency interface
-	virtual void OnAddDependency(const UE::AnimNext::FModuleDependencyContext& InContext) const override;
-	virtual void OnRemoveDependency(const UE::AnimNext::FModuleDependencyContext& InContext) const override;
+	virtual void OnAddDependency(const UE::UAF::FModuleDependencyContext& InContext) const override;
+	virtual void OnRemoveDependency(const UE::UAF::FModuleDependencyContext& InContext) const override;
 
 	// The component on whose primary tick function we will depend
 	UPROPERTY(EditAnywhere, Category = "Dependency")
diff --git a/Engine/Plugins/Experimental/AnimNext/Source/AnimNext/Private/Param/AnimNextObjectFunctionLocatorFragment.cpp b/Engine/Plugins/Experimental/AnimNext/Source/AnimNext/Private/Param/AnimNextObjectFunctionLocatorFragment.cpp
--- a/Engine/Plugins/Experimental/AnimNext/Source/AnimNext/Private/Param/AnimNextObjectFunctionLocatorFragment.cpp
+++ b/Engine/Plugins/Experimental/AnimNext/Source/AnimNext/Private/Param/AnimNextObjectFunctionLocatorFragment.cpp
@@ -23,7 +23,7 @@ UE::UniversalObjectLocator::FResolveResult FAnimNextObjectFunctionLocatorFragmen
 
 	if(UFunction* Function = Cast<UFunction>(Path.ResolveObject()))
 	{
-		if(Params.Context && UE::AnimNext::FParamUtils::CanUseFunction(Function, nullptr))
+		if(Params.Context && UE::UAF::FParamUtils::CanUseFunction(Function, nullptr))
 		{
 			if (FObjectProperty* ObjectProperty = CastField<FObjectProperty>(Function->GetReturnProperty()))
 			{
@@ -64,7 +64,7 @@ UE::UniversalObjectLocator::FInitializeResult FAnimNextObjectFunctionLocatorFrag
 
 	if(const UFunction* Function = Cast<UFunction>(InParams.Object))
 	{
-		if(UE::AnimNext::FParamUtils::CanUseFunction(Function, nullptr))
+		if(UE::UAF::FParamUtils::CanUseFunction(Function, nullptr))
 		{
 			if(FObjectProperty* ObjectProperty = CastField<FObjectProperty>(Function->GetReturnProperty()))
 			{
diff --git a/Engine/Plugins/Experimental/AnimNext/Source/AnimNext/Private/TraitCore/TraitEvent.cpp b/Engine/Plugins/Experimental/AnimNext/Source/AnimNext/Private/TraitCore/TraitEvent.cpp
--- a/Engine/Plugins/Experimental/AnimNext/Source/AnimNext/Private/TraitCore/TraitEvent.cpp
+++ b/Engine/Plugins/Experimental/AnimNext/Source/AnimNext/Private/TraitCore/TraitEvent.cpp
@@ -2,7 +2,7 @@
 
 #include "TraitCore/TraitEvent.h"
 
-bool FAnimNextTraitEvent::DecrementLifetime(UE::AnimNext::FTraitEventList& OutputEventList)
+bool FAnimNextTraitEvent::DecrementLifetime(UE::UAF::FTraitEventList& OutputEventList)
 {
 	const bool bExpired = Lifetime.Decrement();
 	if (bExpired)
diff --git a/Engine/Plugins/Experimental/AnimNext/Source/AnimNext/Private/TransformArray.cpp b/Engine/Plugins/Experimental/AnimNext/Source/AnimNext/Private/TransformArray.cpp
--- a/Engine/Plugins/Experimental/AnimNext/Source/AnimNext/Private/TransformArray.cpp
+++ b/Engine/Plugins/Experimental/AnimNext/Source/AnimNext/Private/TransformArray.cpp
@@ -2,9 +2,9 @@
 
 #include "TransformArray.h"
 
-namespace UE::AnimNext
+namespace UE::UAF
 {
 
 const FTransform TransformAdditiveIdentity = FTransform(FQuat::Identity, FVector::ZeroVector, FVector::ZeroVector);
 
-} // namespace UE::AnimNext
+} // namespace UE::UAF
