UAF: Namespace renamed

diff --git a/Engine/Plugins/Experimental/AnimNext/Source/AnimNext/Private/UAFAssetInstance.cpp b/Engine/Plugins/Experimental/AnimNext/Source/AnimNext/Private/UAFAssetInstance.cpp
--- a/Engine/Plugins/Experimental/AnimNext/Source/AnimNext/Private/UAFAssetInstance.cpp
+++ b/Engine/Plugins/Experimental/AnimNext/Source/AnimNext/Private/UAFAssetInstance.cpp
@@ -9,7 +9,7 @@
 #include "Variables/VariableOverridesCollection.h"
 #include "UAFAssetInstanceComponent.h"
 
-void FUAFAssetInstance::InitializeVariables(const TSharedPtr<const UE::AnimNext::FVariableOverridesCollection>& InOverrides)
+void FUAFAssetInstance::InitializeVariables(const TSharedPtr<const UE::UAF::FVariableOverridesCollection>& InOverrides)
 {
 	Variables.Initialize(*this, InOverrides);
 }
diff --git a/Engine/Plugins/Experimental/AnimNext/Source/AnimNext/Private/UAFAssetInstanceComponent.cpp b/Engine/Plugins/Experimental/AnimNext/Source/AnimNext/Private/UAFAssetInstanceComponent.cpp
--- a/Engine/Plugins/Experimental/AnimNext/Source/AnimNext/Private/UAFAssetInstanceComponent.cpp
+++ b/Engine/Plugins/Experimental/AnimNext/Source/AnimNext/Private/UAFAssetInstanceComponent.cpp
@@ -2,22 +2,22 @@
 
 #include "UAFAssetInstanceComponent.h"
 
-namespace UE::AnimNext::Private
+namespace UE::UAF::Private
 {
 	static thread_local FUAFAssetInstance* GCurrentAssetInstance = nullptr;
 }
 
 FUAFAssetInstanceComponent::FScopedConstructorHelper::FScopedConstructorHelper(FUAFAssetInstance& InInstance)
 {
-	UE::AnimNext::Private::GCurrentAssetInstance = &InInstance;
+	UE::UAF::Private::GCurrentAssetInstance = &InInstance;
 }
 
 FUAFAssetInstanceComponent::FScopedConstructorHelper::~FScopedConstructorHelper()
 {
-	UE::AnimNext::Private::GCurrentAssetInstance = nullptr;
+	UE::UAF::Private::GCurrentAssetInstance = nullptr;
 }
 
 FUAFAssetInstanceComponent::FUAFAssetInstanceComponent()
 {
-	Instance = UE::AnimNext::Private::GCurrentAssetInstance;
+	Instance = UE::UAF::Private::GCurrentAssetInstance;
 }
\ No newline at end of file
diff --git a/Engine/Plugins/Experimental/AnimNext/Source/AnimNext/Private/Variables/AnimNextVariableOverridesCollection.cpp b/Engine/Plugins/Experimental/AnimNext/Source/AnimNext/Private/Variables/AnimNextVariableOverridesCollection.cpp
--- a/Engine/Plugins/Experimental/AnimNext/Source/AnimNext/Private/Variables/AnimNextVariableOverridesCollection.cpp
+++ b/Engine/Plugins/Experimental/AnimNext/Source/AnimNext/Private/Variables/AnimNextVariableOverridesCollection.cpp
@@ -6,7 +6,7 @@
 
 void FAnimNextVariableOverridesCollection::AddStructReferencedObjects(FReferenceCollector& Collector)
 {
-	using namespace UE::AnimNext;
+	using namespace UE::UAF;
 
 	if (TSharedPtr<FVariableOverridesCollection> PinnedCollection = Collection.Pin())
 	{
diff --git a/Engine/Plugins/Experimental/AnimNext/Source/AnimNext/Private/Variables/RigUnit_OverrideVariables.cpp b/Engine/Plugins/Experimental/AnimNext/Source/AnimNext/Private/Variables/RigUnit_OverrideVariables.cpp
--- a/Engine/Plugins/Experimental/AnimNext/Source/AnimNext/Private/Variables/RigUnit_OverrideVariables.cpp
+++ b/Engine/Plugins/Experimental/AnimNext/Source/AnimNext/Private/Variables/RigUnit_OverrideVariables.cpp
@@ -26,10 +26,10 @@ FRigUnit_OverrideVariables_Execute()
 	// We create the overrides one time, then leave the overrides to live as long as the execute context
 	if (!WorkData.Collection.IsValid())
 	{
-		WorkData.Collection = MakeShared<UE::AnimNext::FVariableOverridesCollection>();
+		WorkData.Collection = MakeShared<UE::UAF::FVariableOverridesCollection>();
 
 		FAnimNextModuleInstance& ModuleInstance = ExecuteContext.GetContextData<FAnimNextModuleContextData>().GetModuleInstance();
-		TArray<UE::AnimNext::FVariableOverrides> VariableOverrides;
+		TArray<UE::UAF::FVariableOverrides> VariableOverrides;
 		for (const FRigVMTraitScope& TraitScope : ExecuteContext.GetTraits())
 		{
 			const FRigVMTrait_OverrideVariables* OverrideTrait = TraitScope.GetTrait<FRigVMTrait_OverrideVariables>();
diff --git a/Engine/Plugins/Experimental/AnimNext/Source/AnimNext/Private/Variables/RigUnit_OverrideVariables.h b/Engine/Plugins/Experimental/AnimNext/Source/AnimNext/Private/Variables/RigUnit_OverrideVariables.h
--- a/Engine/Plugins/Experimental/AnimNext/Source/AnimNext/Private/Variables/RigUnit_OverrideVariables.h
+++ b/Engine/Plugins/Experimental/AnimNext/Source/AnimNext/Private/Variables/RigUnit_OverrideVariables.h
@@ -26,7 +26,7 @@ struct FAnimNextOverrideVariablesWorkData
 	GENERATED_BODY()
 
 	// Strong reference held in work data
-	TSharedPtr<UE::AnimNext::FVariableOverridesCollection> Collection;
+	TSharedPtr<UE::UAF::FVariableOverridesCollection> Collection;
 };
 
 /** Hosts a number of traits that allow overrides of variables for specific assets/structs */
diff --git a/Engine/Plugins/Experimental/AnimNext/Source/AnimNext/Private/Variables/RigVMTrait_OverrideVariables.cpp b/Engine/Plugins/Experimental/AnimNext/Source/AnimNext/Private/Variables/RigVMTrait_OverrideVariables.cpp
--- a/Engine/Plugins/Experimental/AnimNext/Source/AnimNext/Private/Variables/RigVMTrait_OverrideVariables.cpp
+++ b/Engine/Plugins/Experimental/AnimNext/Source/AnimNext/Private/Variables/RigVMTrait_OverrideVariables.cpp
@@ -56,7 +56,7 @@ FString FRigVMTrait_OverrideVariablesForAsset::GetDisplayName() const
 }
 #endif
 
-void FRigVMTrait_OverrideVariablesForAsset::GenerateOverrides(const FRigVMTraitScope& InTraitScope, TArray<UE::AnimNext::FVariableOverrides>& OutOverrides) const
+void FRigVMTrait_OverrideVariablesForAsset::GenerateOverrides(const FRigVMTraitScope& InTraitScope, TArray<UE::UAF::FVariableOverrides>& OutOverrides) const
 {
 	if (Asset != nullptr)
 	{
@@ -64,7 +64,7 @@ void FRigVMTrait_OverrideVariablesForAsset::GenerateOverrides(const FRigVMTraitS
 		int32 NumVariables = MemoryHandles.Num();
 		check(Overrides.Num() == NumVariables);
 
-		TArray<UE::AnimNext::FVariableOverrides::FOverride> AssetOverrides;
+		TArray<UE::UAF::FVariableOverrides::FOverride> AssetOverrides;
 		AssetOverrides.Reserve(NumVariables);
 
 		const FInstancedPropertyBag& PropertyBag = Asset->GetVariableDefaults();
@@ -94,7 +94,7 @@ FString FRigVMTrait_OverrideVariablesForStruct::GetDisplayName() const
 }
 #endif
 
-void FRigVMTrait_OverrideVariablesForStruct::GenerateOverrides(const FRigVMTraitScope& InTraitScope, TArray<UE::AnimNext::FVariableOverrides>& OutOverrides) const
+void FRigVMTrait_OverrideVariablesForStruct::GenerateOverrides(const FRigVMTraitScope& InTraitScope, TArray<UE::UAF::FVariableOverrides>& OutOverrides) const
 {
 	if (Struct != nullptr)
 	{
@@ -102,7 +102,7 @@ void FRigVMTrait_OverrideVariablesForStruct::GenerateOverrides(const FRigVMTrait
 		int32 NumVariables = MemoryHandles.Num();
 		check(Overrides.Num() == NumVariables);
 
-		TArray<UE::AnimNext::FVariableOverrides::FOverride> StructOverrides;
+		TArray<UE::UAF::FVariableOverrides::FOverride> StructOverrides;
 		StructOverrides.Reserve(NumVariables);
 
 		for (int32 VariableIndex = 0; VariableIndex < NumVariables; ++VariableIndex)
diff --git a/Engine/Plugins/Experimental/AnimNext/Source/AnimNext/Private/Variables/RigVMTrait_OverrideVariables.h b/Engine/Plugins/Experimental/AnimNext/Source/AnimNext/Private/Variables/RigVMTrait_OverrideVariables.h
--- a/Engine/Plugins/Experimental/AnimNext/Source/AnimNext/Private/Variables/RigVMTrait_OverrideVariables.h
+++ b/Engine/Plugins/Experimental/AnimNext/Source/AnimNext/Private/Variables/RigVMTrait_OverrideVariables.h
@@ -8,7 +8,7 @@
 
 class UAnimNextRigVMAsset;
 
-namespace UE::AnimNext
+namespace UE::UAF
 {
 	struct FVariableOverrides;
 }
@@ -50,7 +50,7 @@ struct ANIMNEXT_API FRigVMTrait_OverrideVariables : public FRigVMTrait
 #endif
 
 	// Override point for derived types to generate overrides
-	virtual void GenerateOverrides(const FRigVMTraitScope& InTraitScope, TArray<UE::AnimNext::FVariableOverrides>& OutOverrides) const {}
+	virtual void GenerateOverrides(const FRigVMTraitScope& InTraitScope, TArray<UE::UAF::FVariableOverrides>& OutOverrides) const {}
 
 	// Cached override data
 	UPROPERTY(EditAnywhere, Category = "Variables", meta=(Input, DetailsOnly, TitleProperty = "Name"))
@@ -65,7 +65,7 @@ struct ANIMNEXT_API FRigVMTrait_OverrideVariablesForAsset : public FRigVMTrait_O
 #if WITH_EDITOR
 	virtual FString GetDisplayName() const override;
 #endif
-	virtual void GenerateOverrides(const FRigVMTraitScope& InTraitScope, TArray<UE::AnimNext::FVariableOverrides>& OutOverrides) const override;
+	virtual void GenerateOverrides(const FRigVMTraitScope& InTraitScope, TArray<UE::UAF::FVariableOverrides>& OutOverrides) const override;
 
 	// The asset to use to generate overrides
 	UPROPERTY(EditAnywhere, Category = "Variables", meta=(Input, Constant, DetailsOnly))
@@ -80,7 +80,7 @@ struct ANIMNEXT_API FRigVMTrait_OverrideVariablesForStruct : public FRigVMTrait_
 #if WITH_EDITOR
 	virtual FString GetDisplayName() const override;
 #endif
-	virtual void GenerateOverrides(const FRigVMTraitScope& InTraitScope, TArray<UE::AnimNext::FVariableOverrides>& OutOverrides) const override;
+	virtual void GenerateOverrides(const FRigVMTraitScope& InTraitScope, TArray<UE::UAF::FVariableOverrides>& OutOverrides) const override;
 
 	// The asset to use to generate overrides
 	UPROPERTY(EditAnywhere, Category = "Variables", meta=(Input, Constant, DetailsOnly))
diff --git a/Engine/Plugins/Experimental/AnimNext/Source/AnimNext/Private/Variables/UAFInstanceVariableContainer.cpp b/Engine/Plugins/Experimental/AnimNext/Source/AnimNext/Private/Variables/UAFInstanceVariableContainer.cpp
--- a/Engine/Plugins/Experimental/AnimNext/Source/AnimNext/Private/Variables/UAFInstanceVariableContainer.cpp
+++ b/Engine/Plugins/Experimental/AnimNext/Source/AnimNext/Private/Variables/UAFInstanceVariableContainer.cpp
@@ -8,7 +8,7 @@
 #include "Param/ParamType.h"
 #include "Variables/VariableOverridesCollection.h"
 
-namespace UE::AnimNext::Private
+namespace UE::UAF::Private
 {
 	// Utility functions for type conversion, adapted from UE::StructUtils::Private in PropertyBag.cpp
 
@@ -559,11 +559,11 @@ namespace UE::AnimNext::Private
 	}
 }
 
-FUAFInstanceVariableContainer::FUAFInstanceVariableContainer(const FUAFAssetInstance& InHostInstance, const UScriptStruct* InStruct, const TSharedPtr<const UE::AnimNext::FVariableOverridesCollection>& InOverrides, int32 InOverridesIndex)
+FUAFInstanceVariableContainer::FUAFInstanceVariableContainer(const FUAFAssetInstance& InHostInstance, const UScriptStruct* InStruct, const TSharedPtr<const UE::UAF::FVariableOverridesCollection>& InOverrides, int32 InOverridesIndex)
 	: HostInstance(&InHostInstance)
 {
 	AssetOrStructData.Set<FStructType>(InStruct);
-	StructDataCache = UE::AnimNext::FStructDataCache::GetStructInfo(InStruct);
+	StructDataCache = UE::UAF::FStructDataCache::GetStructInfo(InStruct);
 	NumVariables = StructDataCache->GetProperties().Num();
 	VariablesContainer.Set<FInstancedStruct>(FInstancedStruct());
 	VariablesContainer.Get<FInstancedStruct>().InitializeAs(InStruct);
@@ -573,7 +573,7 @@ FUAFInstanceVariableContainer::FUAFInstanceVariableContainer(const FUAFAssetInst
 	ResolveOverrides();
 }
 
-FUAFInstanceVariableContainer::FUAFInstanceVariableContainer(const FUAFAssetInstance& InHostInstance, const UAnimNextRigVMAsset* InAsset, const TSharedPtr<const UE::AnimNext::FVariableOverridesCollection>& InOverrides, int32 InOverridesIndex)
+FUAFInstanceVariableContainer::FUAFInstanceVariableContainer(const FUAFAssetInstance& InHostInstance, const UAnimNextRigVMAsset* InAsset, const TSharedPtr<const UE::UAF::FVariableOverridesCollection>& InOverrides, int32 InOverridesIndex)
 	: HostInstance(&InHostInstance)
 {
 	AssetOrStructData.Set<FAssetType>(InAsset);
@@ -587,7 +587,7 @@ FUAFInstanceVariableContainer::FUAFInstanceVariableContainer(const FUAFAssetInst
 
 EPropertyBagResult FUAFInstanceVariableContainer::GetVariable(int32 InVariableIndex, const FAnimNextParamType& InType, TArrayView<uint8> OutResult) const
 {
-	using namespace UE::AnimNext;
+	using namespace UE::UAF;
 
 	const bool bHasOverrides = ResolvedOverrides.Num() > 0;
 	FAnimNextParamType FoundType;
@@ -654,7 +654,7 @@ EPropertyBagResult FUAFInstanceVariableContainer::GetVariable(int32 InVariableIn
 
 EPropertyBagResult FUAFInstanceVariableContainer::AccessVariable(int32 InVariableIndex, const FAnimNextParamType& InType, TArrayView<uint8>& OutResult)
 {
-	using namespace UE::AnimNext;
+	using namespace UE::UAF;
 
 	const bool bHasOverrides = ResolvedOverrides.Num() > 0;
 	FAnimNextParamType FoundType;
@@ -718,7 +718,7 @@ EPropertyBagResult FUAFInstanceVariableContainer::AccessVariable(int32 InVariabl
 
 void FUAFInstanceVariableContainer::AccessVariableUnchecked(int32 InVariableIndex, const FProperty*& OutProperty, TArrayView<uint8>& OutResult)
 {
-	using namespace UE::AnimNext;
+	using namespace UE::UAF;
 
 	const bool bHasOverrides = ResolvedOverrides.Num() > 0;
 	int32 Size = 0;
@@ -775,7 +775,7 @@ void FUAFInstanceVariableContainer::AccessVariableUnchecked(int32 InVariableInde
 
 EPropertyBagResult FUAFInstanceVariableContainer::SetVariable(int32 InVariableIndex, const FAnimNextParamType& InType, TConstArrayView<uint8> InNewValue)
 {
-	using namespace UE::AnimNext;
+	using namespace UE::UAF;
 
 	const bool bHasOverrides = ResolvedOverrides.Num() > 0;
 	FAnimNextParamType FoundType;
@@ -841,7 +841,7 @@ EPropertyBagResult FUAFInstanceVariableContainer::SetVariable(int32 InVariableIn
 
 void FUAFInstanceVariableContainer::GetAllVariablesOfType(const FAnimNextParamType& InType, TArray<FAnimNextVariableReference>& OutVariables) const
 {
-	using namespace UE::AnimNext;
+	using namespace UE::UAF;
 
 	if (NumVariables == 0)
 	{
@@ -886,7 +886,7 @@ void FUAFInstanceVariableContainer::GetAllVariablesOfType(const FAnimNextParamTy
 
 bool FUAFInstanceVariableContainer::ForEachVariable(TFunctionRef<bool(FName, const FAnimNextParamType&, int32)> InFunction, int32& InOutVariableIndex) const
 {
-	using namespace UE::AnimNext;
+	using namespace UE::UAF;
 	
 	if (NumVariables == 0)
 	{
@@ -932,7 +932,7 @@ bool FUAFInstanceVariableContainer::ForEachVariable(TFunctionRef<bool(FName, con
 
 FAnimNextParamType FUAFInstanceVariableContainer::GetVariableType(int32 InVariableIndex) const
 {
-	using namespace UE::AnimNext;
+	using namespace UE::UAF;
 
 	switch (VariablesContainer.GetIndex())
 	{
@@ -959,7 +959,7 @@ FAnimNextParamType FUAFInstanceVariableContainer::GetVariableType(int32 InVariab
 
 void FUAFInstanceVariableContainer::ResolveOverrides()
 {
-	using namespace UE::AnimNext;
+	using namespace UE::UAF;
 
 	TSharedPtr<const FVariableOverridesCollection> PinnedOverrides = Overrides.Pin();
 	if (!PinnedOverrides.IsValid())
@@ -1042,7 +1042,7 @@ void FUAFInstanceVariableContainer::ResolveOverrides()
 #if WITH_EDITOR
 void FUAFInstanceVariableContainer::Migrate()
 {
-	using namespace UE::AnimNext;
+	using namespace UE::UAF;
 	
 	switch (VariablesContainer.GetIndex())
 	{
@@ -1070,7 +1070,7 @@ void FUAFInstanceVariableContainer::Migrate()
 
 void FUAFInstanceVariableContainer::GetRigVMMemoryForVariables(TConstArrayView<FRigVMExternalVariableDef> InVariableDefs, TArray<FRigVMExternalVariableRuntimeData>& OutRuntimeData)
 {
-	using namespace UE::AnimNext;
+	using namespace UE::UAF;
 	
 	if (NumVariables == 0)
 	{
diff --git a/Engine/Plugins/Experimental/AnimNext/Source/AnimNext/Private/Variables/UAFInstanceVariableData.cpp b/Engine/Plugins/Experimental/AnimNext/Source/AnimNext/Private/Variables/UAFInstanceVariableData.cpp
--- a/Engine/Plugins/Experimental/AnimNext/Source/AnimNext/Private/Variables/UAFInstanceVariableData.cpp
+++ b/Engine/Plugins/Experimental/AnimNext/Source/AnimNext/Private/Variables/UAFInstanceVariableData.cpp
@@ -11,9 +11,9 @@
 #include "VariableOverridesCollection.h"
 #include "Variables/UAFInstanceVariableContainer.h"
 
-void FUAFInstanceVariableData::Initialize(const FUAFAssetInstance& InInstance, const TSharedPtr<const UE::AnimNext::FVariableOverridesCollection>& InOverrides)
+void FUAFInstanceVariableData::Initialize(const FUAFAssetInstance& InInstance, const TSharedPtr<const UE::UAF::FVariableOverridesCollection>& InOverrides)
 {
-	using namespace UE::AnimNext;
+	using namespace UE::UAF;
 
 	Reset();
 
@@ -178,7 +178,7 @@ void FUAFInstanceVariableData::Initialize(const FUAFAssetInstance& InInstance, c
 
 void FUAFInstanceVariableData::RebuildNameMaps()
 {
-	using namespace UE::AnimNext;
+	using namespace UE::UAF;
 
 	LegacyVariableNameMap.Reset();
 	VariablesMap.Reset();
@@ -323,7 +323,7 @@ const FUAFInstanceVariableData::FVariableMapping* FUAFInstanceVariableData::Find
 
 EPropertyBagResult FUAFInstanceVariableData::GetVariable(const FAnimNextVariableReference& InVariable, const FAnimNextParamType& InType, TArrayView<uint8> OutResult) const
 {
-	using namespace UE::AnimNext;
+	using namespace UE::UAF;
 
 	const FVariableMapping* FoundMapping = FindMapping(InVariable);
 	if (FoundMapping == nullptr)
@@ -339,7 +339,7 @@ EPropertyBagResult FUAFInstanceVariableData::GetVariable(const FAnimNextVariable
 
 EPropertyBagResult FUAFInstanceVariableData::AccessVariable(const FAnimNextVariableReference& InVariable, const FAnimNextParamType& InType, TArrayView<uint8>& OutResult) const
 {
-	using namespace UE::AnimNext;
+	using namespace UE::UAF;
 
 	const FVariableMapping* FoundMapping = FindMapping(InVariable);
 	if (FoundMapping == nullptr)
@@ -355,7 +355,7 @@ EPropertyBagResult FUAFInstanceVariableData::AccessVariable(const FAnimNextVaria
 
 EPropertyBagResult FUAFInstanceVariableData::SetVariable(const FAnimNextVariableReference& InVariable, const FAnimNextParamType& InType, TConstArrayView<uint8> InNewValue) const
 {
-	using namespace UE::AnimNext;
+	using namespace UE::UAF;
 
 	const FVariableMapping* FoundMapping = FindMapping(InVariable);
 	if (FoundMapping == nullptr)
diff --git a/Engine/Plugins/Experimental/AnimNext/Source/AnimNext/Public/AnimNextRigVMAsset.h b/Engine/Plugins/Experimental/AnimNext/Source/AnimNext/Public/AnimNextRigVMAsset.h
--- a/Engine/Plugins/Experimental/AnimNext/Source/AnimNext/Public/AnimNextRigVMAsset.h
+++ b/Engine/Plugins/Experimental/AnimNext/Source/AnimNext/Public/AnimNextRigVMAsset.h
@@ -22,7 +22,7 @@ struct FUAFInstanceVariableDataProxy;
 struct FAnimNextRigVMFunctionData;
 struct FUAFRigVMComponent;
 
-namespace UE::AnimNext
+namespace UE::UAF
 {
 	struct FAnimGraphBuilderContext;
 	struct FFunctionHandle;
@@ -31,7 +31,7 @@ namespace UE::AnimNext
 	struct FInjectionInfo;
 }
 
-namespace UE::AnimNext::UncookedOnly
+namespace UE::UAF::UncookedOnly
 {
 	struct FUtils;
 }
@@ -47,11 +47,11 @@ protected:
 	friend class UAnimNextRigVMAssetEditorData;
 	friend UAnimNextAnimGraphSettings;
 	friend FUAFRigVMComponent;
-	friend struct UE::AnimNext::FModuleEventTickFunction;
-	friend struct UE::AnimNext::UncookedOnly::FUtils;
+	friend struct UE::UAF::FModuleEventTickFunction;
+	friend struct UE::UAF::UncookedOnly::FUtils;
 	friend FUAFInstanceVariableData;
-	friend UE::AnimNext::FInjectionInfo;
-	friend UE::AnimNext::FAnimGraphBuilderContext;
+	friend UE::UAF::FInjectionInfo;
+	friend UE::UAF::FAnimGraphBuilderContext;
 
 	UE_API UAnimNextRigVMAsset(const FObjectInitializer& ObjectInitializer);
 
@@ -69,10 +69,10 @@ public:
 	const URigVM* GetRigVM() const { return RigVM; }
 
 	// Call a function handle using a property bag as arguments.
-	UE_API void CallFunctionHandle(UE::AnimNext::FFunctionHandle InHandle, FRigVMExtendedExecuteContext& InContext, FInstancedPropertyBag& InArgs);
+	UE_API void CallFunctionHandle(UE::UAF::FFunctionHandle InHandle, FRigVMExtendedExecuteContext& InContext, FInstancedPropertyBag& InArgs);
 
 	// Get a function handle given its mangled event name
-	UE_API UE::AnimNext::FFunctionHandle GetFunctionHandle(FName InEventName) const;
+	UE_API UE::UAF::FFunctionHandle GetFunctionHandle(FName InEventName) const;
 
 #if WITH_EDITOR
 	// Delegate called in editor when an asset's compilation has started (i.e. asset data that instances rely on is going to change)
diff --git a/Engine/Plugins/Experimental/AnimNext/Source/AnimNext/Public/Component/AnimNextComponent.h b/Engine/Plugins/Experimental/AnimNext/Source/AnimNext/Public/Component/AnimNextComponent.h
--- a/Engine/Plugins/Experimental/AnimNext/Source/AnimNext/Public/Component/AnimNextComponent.h
+++ b/Engine/Plugins/Experimental/AnimNext/Source/AnimNext/Public/Component/AnimNextComponent.h
@@ -17,7 +17,7 @@ class UAnimNextComponentWorldSubsystem;
 class UAnimNextModule;
 struct AnimNextVariableReference;
 
-namespace UE::AnimNext::UncookedOnly
+namespace UE::UAF::UncookedOnly
 {
 	struct FUtils;
 }
@@ -85,14 +85,14 @@ public:
 	ANIMNEXT_API void ShowDebugDrawing(bool bShowDebugDrawing);
 
 	// Queue a task to run during execution 
-	ANIMNEXT_API void QueueTask(FName InModuleEventName, TUniqueFunction<void(const UE::AnimNext::FModuleTaskContext&)>&& InTaskFunction, UE::AnimNext::ETaskRunLocation InLocation = UE::AnimNext::ETaskRunLocation::Before);
+	ANIMNEXT_API void QueueTask(FName InModuleEventName, TUniqueFunction<void(const UE::UAF::FModuleTaskContext&)>&& InTaskFunction, UE::UAF::ETaskRunLocation InLocation = UE::UAF::ETaskRunLocation::Before);
 
 	// Queues an input trait event
 	// Input events will be processed in the next graph update after they are queued
 	ANIMNEXT_API void QueueInputTraitEvent(FAnimNextTraitEventPtr Event);
 
 	// Get the handle to the registered module
-	UE::AnimNext::FModuleHandle GetModuleHandle() const { return ModuleHandle; }
+	UE::UAF::FModuleHandle GetModuleHandle() const { return ModuleHandle; }
 
 	// Find the tick function for the specified event
 	// @param	InEventName			The event associated to the wanted tick function
@@ -190,7 +190,7 @@ private:
 	ANIMNEXT_API EPropertyBagResult WriteVariableInternal(const FAnimNextVariableReference& InVariable, const FAnimNextParamType& InType, TFunctionRef<void(TArrayView<uint8>)> InFunction);
 
 private:
-	friend struct UE::AnimNext::UncookedOnly::FUtils;
+	friend struct UE::UAF::UncookedOnly::FUtils;
 	friend class UAnimNextComponentWorldSubsystem;
 
 	// The AnimNext module that this component will run
@@ -201,7 +201,7 @@ private:
 	TObjectPtr<UAnimNextComponentWorldSubsystem> Subsystem = nullptr;
 
 	// Handle to the registered module
-	UE::AnimNext::FModuleHandle ModuleHandle;
+	UE::UAF::FModuleHandle ModuleHandle;
 
 	// How to initialize the module
 	UPROPERTY(EditAnywhere, Category="Module")
diff --git a/Engine/Plugins/Experimental/AnimNext/Source/AnimNext/Public/GenerationTools.h b/Engine/Plugins/Experimental/AnimNext/Source/AnimNext/Public/GenerationTools.h
--- a/Engine/Plugins/Experimental/AnimNext/Source/AnimNext/Public/GenerationTools.h
+++ b/Engine/Plugins/Experimental/AnimNext/Source/AnimNext/Public/GenerationTools.h
@@ -16,7 +16,7 @@ class FSkeletalMeshLODRenderData;
 struct FPoseContext;
 class UPhysicsAsset;
 
-namespace UE::AnimNext
+namespace UE::UAF
 {
 
 struct FGenerationLODData
@@ -35,7 +35,7 @@ public:
 	// If no SkeletalMesh asset is passed, there will be no generation
 	static UE_API bool GenerateReferencePose(const USkeletalMeshComponent* SkeletalMeshComponent
 		, USkeletalMesh* SkeletalMesh
-		, UE::AnimNext::FReferencePose& OutAnimationReferencePose);
+		, UE::UAF::FReferencePose& OutAnimationReferencePose);
 
 	// Generates the full list of bones required by a LOD, based on the SkeletalRetrieve RequiredBones
 	// Compute ExcludedBones versus LOD0 and PreviousLOD
@@ -155,6 +155,6 @@ public:
 	static UE_API void ConvertLocalSpaceToComponentSpace(TConstArrayView<FBoneIndexType> InParentIndices, TConstArrayView<FTransform> InBoneSpaceTransforms, TConstArrayView<FBoneIndexType> InRequiredBoneIndices, TArrayView<FTransform> OutComponentSpaceTransforms);
 };
 
-} // namespace UE::AnimNext
+} // namespace UE::UAF
 
 #undef UE_API
diff --git a/Engine/Plugins/Experimental/AnimNext/Source/AnimNext/Public/InstanceTaskContext.h b/Engine/Plugins/Experimental/AnimNext/Source/AnimNext/Public/InstanceTaskContext.h
--- a/Engine/Plugins/Experimental/AnimNext/Source/AnimNext/Public/InstanceTaskContext.h
+++ b/Engine/Plugins/Experimental/AnimNext/Source/AnimNext/Public/InstanceTaskContext.h
@@ -12,12 +12,12 @@ struct FAnimNextVariableReference;
 struct FAnimNextFactoryParams;
 class UInjectionCallbackProxy;
 
-namespace UE::AnimNext
+namespace UE::UAF
 {
 	struct FInjectionSiteTrait;
 }
 
-namespace UE::AnimNext
+namespace UE::UAF
 {
 
 // Context struct passed to instance tasks, allowing setting of variables etc.
diff --git a/Engine/Plugins/Experimental/AnimNext/Source/AnimNext/Public/LODPose.h b/Engine/Plugins/Experimental/AnimNext/Source/AnimNext/Public/LODPose.h
--- a/Engine/Plugins/Experimental/AnimNext/Source/AnimNext/Public/LODPose.h
+++ b/Engine/Plugins/Experimental/AnimNext/Source/AnimNext/Public/LODPose.h
@@ -12,7 +12,7 @@
 // Handy define to set uninitialized pose buffers to 0xCDCDCDCD to help find issues
 #define UE_ENABLE_POSE_DEBUG_FILL 0
 
-namespace UE::AnimNext
+namespace UE::UAF
 {
 
 enum class ELODPoseFlags : uint8
@@ -137,7 +137,7 @@ struct FLODPose
 
 	void SetIdentity(bool bAdditive = false, int32 StartIndex = 0, int32 NumTransformsToSet = -1)
 	{
-		UE::AnimNext::SetIdentity(LocalTransformsView, bAdditive, StartIndex, NumTransformsToSet);
+		UE::UAF::SetIdentity(LocalTransformsView, bAdditive, StartIndex, NumTransformsToSet);
 	}
 
 	int32 GetNumBones() const
@@ -380,7 +380,7 @@ using FLODPoseStack = TLODPose<FAnimStackAllocator>;
 USTRUCT()
 struct FAnimNextLODPose
 #if CPP
-	: UE::AnimNext::FLODPoseHeap
+	: UE::UAF::FLODPoseHeap
 #endif
 {
 	GENERATED_BODY()
diff --git a/Engine/Plugins/Experimental/AnimNext/Source/AnimNext/Public/Module/ModuleHandle.h b/Engine/Plugins/Experimental/AnimNext/Source/AnimNext/Public/Module/ModuleHandle.h
--- a/Engine/Plugins/Experimental/AnimNext/Source/AnimNext/Public/Module/ModuleHandle.h
+++ b/Engine/Plugins/Experimental/AnimNext/Source/AnimNext/Public/Module/ModuleHandle.h
@@ -8,7 +8,7 @@
 
 struct FAnimNextModuleInstance;
 
-namespace UE::AnimNext
+namespace UE::UAF
 {
 
 // Opaque handle representing a module instance
@@ -23,11 +23,11 @@ struct FAnimNextModuleHandle
 
 	FAnimNextModuleHandle() = default;
 
-	explicit FAnimNextModuleHandle(const UE::AnimNext::FModuleHandle InModuleHandle)
+	explicit FAnimNextModuleHandle(const UE::UAF::FModuleHandle InModuleHandle)
 		: ModuleHandle(InModuleHandle)
 	{
 	}
 
 public:
-	UE::AnimNext::FModuleHandle ModuleHandle;
+	UE::UAF::FModuleHandle ModuleHandle;
 };
\ No newline at end of file
diff --git a/Engine/Plugins/Experimental/AnimNext/Source/AnimNext/Public/Module/ModuleTaskContext.h b/Engine/Plugins/Experimental/AnimNext/Source/AnimNext/Public/Module/ModuleTaskContext.h
--- a/Engine/Plugins/Experimental/AnimNext/Source/AnimNext/Public/Module/ModuleTaskContext.h
+++ b/Engine/Plugins/Experimental/AnimNext/Source/AnimNext/Public/Module/ModuleTaskContext.h
@@ -12,14 +12,14 @@ struct FAnimNextModuleInstance;
 struct FInstancedPropertyBag;
 struct FAnimNextModuleInjectionComponent;
 
-namespace UE::AnimNext
+namespace UE::UAF
 {
 	struct FScheduleContext;
 	enum class EParameterScopeOrdering : int32;
 	struct FModuleEventTickFunction;
 }
 
-namespace UE::AnimNext
+namespace UE::UAF
 {
 
 // Context passed to schedule task callbacks
@@ -50,7 +50,7 @@ private:
 	// The module instance currently running
 	FAnimNextModuleInstance& ModuleInstance;
 
-	friend UE::AnimNext::FModuleEventTickFunction;
+	friend UE::UAF::FModuleEventTickFunction;
 	friend ::FAnimNextModuleInjectionComponent;
 };
 
diff --git a/Engine/Plugins/Experimental/AnimNext/Source/AnimNext/Public/Module/RigVMTrait_ModuleEventDependency.h b/Engine/Plugins/Experimental/AnimNext/Source/AnimNext/Public/Module/RigVMTrait_ModuleEventDependency.h
--- a/Engine/Plugins/Experimental/AnimNext/Source/AnimNext/Public/Module/RigVMTrait_ModuleEventDependency.h
+++ b/Engine/Plugins/Experimental/AnimNext/Source/AnimNext/Public/Module/RigVMTrait_ModuleEventDependency.h
@@ -21,7 +21,7 @@ enum class EAnimNextModuleEventDependencyOrdering : uint8
 	After
 };
 
-namespace UE::AnimNext
+namespace UE::UAF
 {
 	// Context passed to module dependency functions
 	struct FModuleDependencyContext
@@ -48,10 +48,10 @@ struct FRigVMTrait_ModuleEventDependency : public FRigVMTrait
 	UE_API FRigVMTrait_ModuleEventDependency();
 
 	// Override point that adds the dependency
-	virtual void OnAddDependency(const UE::AnimNext::FModuleDependencyContext& InContext) const {}
+	virtual void OnAddDependency(const UE::UAF::FModuleDependencyContext& InContext) const {}
 
 	// Override point that removes the dependency
-	virtual void OnRemoveDependency(const UE::AnimNext::FModuleDependencyContext& InContext) const {}
+	virtual void OnRemoveDependency(const UE::UAF::FModuleDependencyContext& InContext) const {}
 
 	// How to execute relative to the event
 	UPROPERTY(EditAnywhere, Category="Dependency")
diff --git a/Engine/Plugins/Experimental/AnimNext/Source/AnimNext/Public/Param/ParamType.h b/Engine/Plugins/Experimental/AnimNext/Source/AnimNext/Public/Param/ParamType.h
--- a/Engine/Plugins/Experimental/AnimNext/Source/AnimNext/Public/Param/ParamType.h
+++ b/Engine/Plugins/Experimental/AnimNext/Source/AnimNext/Public/Param/ParamType.h
@@ -14,17 +14,17 @@
 
 #define UE_API ANIMNEXT_API
 
-namespace UE::AnimNext
+namespace UE::UAF
 {
 	struct FParamHelpers;
 }
 
-namespace UE::AnimNext::UncookedOnly
+namespace UE::UAF::UncookedOnly
 {
 	struct FUtils;
 }
 
-namespace UE::AnimNext::Private
+namespace UE::UAF::Private
 {
 	template <typename T>
 	struct TIsSoftObjectPtr
@@ -51,7 +51,7 @@ namespace UE::AnimNext::Private
 
 /**
  * Representation of a parameter's type. Serializable, but fairly heavyweight to pass around and compare.
- * Faster comparisons and other operations can be performed on UE::AnimNext::FParamTypeHandle, but they cannot be
+ * Faster comparisons and other operations can be performed on UE::UAF::FParamTypeHandle, but they cannot be
  * serialized as they are not stable across runs.
  */
 USTRUCT()
@@ -63,8 +63,8 @@ public:
 	using EValueType = ::EPropertyBagPropertyType;
 	using EContainerType = ::EPropertyBagContainerType;
 
-	friend struct UE::AnimNext::FParamHelpers;
-	friend struct UE::AnimNext::UncookedOnly::FUtils;
+	friend struct UE::UAF::FParamHelpers;
+	friend struct UE::UAF::UncookedOnly::FUtils;
 
 	FAnimNextParamType() = default;
 
@@ -170,11 +170,11 @@ private:
 		{
 			return FAnimNextParamType(EValueType::Class, InContainerType, ParamType::ElementType::StaticClass());
 		}
-		else if constexpr (UE::AnimNext::Private::TIsSoftObjectPtr<ParamType>::Value)
+		else if constexpr (UE::UAF::Private::TIsSoftObjectPtr<ParamType>::Value)
 		{
 			return FAnimNextParamType(EValueType::SoftObject, InContainerType, ParamType::ElementType::StaticClass());
 		}
-		else if constexpr (UE::AnimNext::Private::TIsSoftClassPtr<ParamType>::Value)
+		else if constexpr (UE::UAF::Private::TIsSoftClassPtr<ParamType>::Value)
 		{
 			return FAnimNextParamType(EValueType::SoftClass, InContainerType, ParamType::ElementType::StaticClass());
 		}
diff --git a/Engine/Plugins/Experimental/AnimNext/Source/AnimNext/Public/ReferencePose.h b/Engine/Plugins/Experimental/AnimNext/Source/AnimNext/Public/ReferencePose.h
--- a/Engine/Plugins/Experimental/AnimNext/Source/AnimNext/Public/ReferencePose.h
+++ b/Engine/Plugins/Experimental/AnimNext/Source/AnimNext/Public/ReferencePose.h
@@ -17,7 +17,7 @@
 class USkeletalMesh;
 class USkeleton;
 
-namespace UE::AnimNext
+namespace UE::UAF
 {
 
 enum class EReferencePoseGenerationFlags : uint8
@@ -334,13 +334,13 @@ struct TReferencePose
 */
 using FReferencePose = TReferencePose<FDefaultAllocator, FDefaultSetAllocator>;
 
-} // namespace UE::AnimNext
+} // namespace UE::UAF
 
 // USTRUCT wrapper for reference pose
 USTRUCT()
 struct FAnimNextReferencePose
 #if CPP
-	: UE::AnimNext::FReferencePose
+	: UE::UAF::FReferencePose
 #endif
 {
 	GENERATED_BODY()
diff --git a/Engine/Plugins/Experimental/AnimNext/Source/AnimNext/Public/RewindDebugger/AnimNextTrace.h b/Engine/Plugins/Experimental/AnimNext/Source/AnimNext/Public/RewindDebugger/AnimNextTrace.h
--- a/Engine/Plugins/Experimental/AnimNext/Source/AnimNext/Public/RewindDebugger/AnimNextTrace.h
+++ b/Engine/Plugins/Experimental/AnimNext/Source/AnimNext/Public/RewindDebugger/AnimNextTrace.h
@@ -26,7 +26,7 @@ struct FUAFAssetInstance;
 struct FAnimNextGraphInstance;
 struct FUAFInstanceVariableContainer;
 
-namespace UE::AnimNext
+namespace UE::UAF
 {
 struct FAnimNextTrace
 {
@@ -44,8 +44,8 @@ private:
 }
 
 
-#define TRACE_ANIMNEXT_INSTANCE(DataInterface, OuterObject) UE::AnimNext::FAnimNextTrace::OutputAnimNextInstance(DataInterface, OuterObject);
-#define TRACE_ANIMNEXT_VARIABLES(DataInterface, OuterObject) UE::AnimNext::FAnimNextTrace::OutputAnimNextVariables(DataInterface, OuterObject);
+#define TRACE_ANIMNEXT_INSTANCE(DataInterface, OuterObject) UE::UAF::FAnimNextTrace::OutputAnimNextInstance(DataInterface, OuterObject);
+#define TRACE_ANIMNEXT_VARIABLES(DataInterface, OuterObject) UE::UAF::FAnimNextTrace::OutputAnimNextVariables(DataInterface, OuterObject);
 #else
 #define TRACE_ANIMNEXT_INSTANCE(DataInterface, OuterObject)
 #define TRACE_ANIMNEXT_VARIABLES(DataInterface, OuterObject)
diff --git a/Engine/Plugins/Experimental/AnimNext/Source/AnimNext/Public/TraitCore/TraitEvent.h b/Engine/Plugins/Experimental/AnimNext/Source/AnimNext/Public/TraitCore/TraitEvent.h
--- a/Engine/Plugins/Experimental/AnimNext/Source/AnimNext/Public/TraitCore/TraitEvent.h
+++ b/Engine/Plugins/Experimental/AnimNext/Source/AnimNext/Public/TraitCore/TraitEvent.h
@@ -9,21 +9,21 @@
 
 #include "TraitEvent.generated.h"
 
-namespace UE::AnimNext
+namespace UE::UAF
 {
 	struct FTraitEventList;
 
 	// Create an alias to improve readability
-	using FTraitEventUID = UE::AnimNext::FConstExprUID;
+	using FTraitEventUID = UE::UAF::FConstExprUID;
 }
 
 // Helper macros
 // In a trait event struct declaration, this macro declares the necessary boilerplate we require
 #define DECLARE_ANIM_TRAIT_EVENT(EventName, SuperEventName) \
 	/* FAnimNextTraitEvent impl */ \
-	static constexpr UE::AnimNext::FTraitEventUID TypeUID = UE::AnimNext::FTraitEventUID::MakeFromString(TEXT(#EventName)); \
-	virtual UE::AnimNext::FTraitEventUID GetTypeUID() const override { return TypeUID; } \
-	virtual bool IsA(UE::AnimNext::FTraitEventUID InTypeUID) const override { return InTypeUID == TypeUID ? true : SuperEventName::IsA(InTypeUID); } \
+	static constexpr UE::UAF::FTraitEventUID TypeUID = UE::UAF::FTraitEventUID::MakeFromString(TEXT(#EventName)); \
+	virtual UE::UAF::FTraitEventUID GetTypeUID() const override { return TypeUID; } \
+	virtual bool IsA(UE::UAF::FTraitEventUID InTypeUID) const override { return InTypeUID == TypeUID ? true : SuperEventName::IsA(InTypeUID); } \
 
 
 /**
@@ -45,24 +45,24 @@ struct FAnimNextTraitEvent
 	virtual ~FAnimNextTraitEvent() = default;
 
 	// Returns the event type UID
-	virtual UE::AnimNext::FTraitEventUID GetTypeUID() const { return UE::AnimNext::FTraitEventUID(); }
+	virtual UE::UAF::FTraitEventUID GetTypeUID() const { return UE::UAF::FTraitEventUID(); }
 
 	// Returns whether or not this event derives from the specified type
 	template<class EventType>
 	bool IsA() const { return IsA(EventType::TypeUID); }
 
 	// Returns whether or not this event derives from the specified type
-	virtual bool IsA(UE::AnimNext::FTraitEventUID InTypeUID) const { return false; }
+	virtual bool IsA(UE::UAF::FTraitEventUID InTypeUID) const { return false; }
 
 	// Returns whether or not this event is valid
 	// An event is valid if it hasn't been consumed and if it isn't expired
 	bool IsValid() const { return !IsConsumed() && !IsExpired(); }
 
 	// Returns the lifetime of this event
-	UE::AnimNext::FTraitEventLifetime GetLifetime() const { return Lifetime; }
+	UE::UAF::FTraitEventLifetime GetLifetime() const { return Lifetime; }
 
 	// Sets the desired lifetime for this event
-	void SetLifetime(UE::AnimNext::FTraitEventLifetime InLifetime) { Lifetime = InLifetime; }
+	void SetLifetime(UE::UAF::FTraitEventLifetime InLifetime) { Lifetime = InLifetime; }
 
 	// Marks this event as being handled.
 	// Subsequent handlers can use this to determine if they should take action or not.
@@ -91,7 +91,7 @@ struct FAnimNextTraitEvent
 
 	// Decrements the lifetime count and returns whether or not this event has expired in the process
 	// OnExpired will be called on the derived type if we expire and new output events can be appended
-	ANIMNEXT_API bool DecrementLifetime(UE::AnimNext::FTraitEventList& OutputEventList);
+	ANIMNEXT_API bool DecrementLifetime(UE::UAF::FTraitEventList& OutputEventList);
 
 	// Returns a pointer to this event cast to the specified type if the TypeUIDs match, nullptr otherwise
 	template<class EventType>
@@ -103,17 +103,17 @@ struct FAnimNextTraitEvent
 
 protected:
 	// Creates an event of the specified type and associated lifetime
-	explicit FAnimNextTraitEvent(UE::AnimNext::FTraitEventLifetime InLifetime)
+	explicit FAnimNextTraitEvent(UE::UAF::FTraitEventLifetime InLifetime)
 		: Lifetime(InLifetime)
 	{
 	}
 
 	// Called when an events' lifetime expires
-	virtual void OnExpired(UE::AnimNext::FTraitEventList& OutputEventList) {}
+	virtual void OnExpired(UE::UAF::FTraitEventList& OutputEventList) {}
 
 private:
 	// The specified lifetime of this event
-	UE::AnimNext::FTraitEventLifetime Lifetime = UE::AnimNext::FTraitEventLifetime::MakeTransient();
+	UE::UAF::FTraitEventLifetime Lifetime = UE::UAF::FTraitEventLifetime::MakeTransient();
 
 	// Signals whether someone has opted to handle this event
 	int8 bIsHandled = 0;
diff --git a/Engine/Plugins/Experimental/AnimNext/Source/AnimNext/Public/UAFAssetInstance.h b/Engine/Plugins/Experimental/AnimNext/Source/AnimNext/Public/UAFAssetInstance.h
--- a/Engine/Plugins/Experimental/AnimNext/Source/AnimNext/Public/UAFAssetInstance.h
+++ b/Engine/Plugins/Experimental/AnimNext/Source/AnimNext/Public/UAFAssetInstance.h
@@ -17,7 +17,7 @@ struct FAnimNextModuleInjectionComponent;
 struct FUAFAssetInstanceComponent;
 struct FUAFRigVMComponent;
 
-namespace UE::AnimNext
+namespace UE::UAF
 {
 	struct FPoolHandle;
 	struct FInjectionInfo;
@@ -250,7 +250,7 @@ struct FUAFAssetInstance
 
 protected:
 	// Used during initialization. Creates variables or references those of the outer host and applies any overrides
-	UE_API void InitializeVariables(const TSharedPtr<const UE::AnimNext::FVariableOverridesCollection>& InOverrides = nullptr);
+	UE_API void InitializeVariables(const TSharedPtr<const UE::UAF::FVariableOverridesCollection>& InOverrides = nullptr);
 
 #if WITH_EDITOR
 	// Used during compilation, migrate all variables to new bags according to new defaults. Reapplies any cached overrides.
@@ -285,11 +285,11 @@ protected:
 	// Cached RigVM component, used to avoid map lookup each time we execute 
 	FUAFRigVMComponent* CachedRigVMComponent = nullptr;
 
-	friend UE::AnimNext::FInjectionInfo;
+	friend UE::UAF::FInjectionInfo;
 	friend FAnimNextModuleInjectionComponent;
 	friend FUAFInstanceVariableData;
-	friend UE::AnimNext::FAnimNextTrace;
-	friend UE::AnimNext::FInstanceTaskContext;
+	friend UE::UAF::FAnimNextTrace;
+	friend UE::UAF::FInstanceTaskContext;
 	friend FUAFRigVMComponent;
 
 #if ANIMNEXT_TRACE_ENABLED
diff --git a/Engine/Plugins/Experimental/AnimNext/Source/AnimNext/Public/Variables/AnimNextSharedVariables.h b/Engine/Plugins/Experimental/AnimNext/Source/AnimNext/Public/Variables/AnimNextSharedVariables.h
--- a/Engine/Plugins/Experimental/AnimNext/Source/AnimNext/Public/Variables/AnimNextSharedVariables.h
+++ b/Engine/Plugins/Experimental/AnimNext/Source/AnimNext/Public/Variables/AnimNextSharedVariables.h
@@ -9,12 +9,12 @@
 class UAnimNextSharedVariablesFactory;
 class UAnimNextSharedVariables;
 
-namespace UE::AnimNext
+namespace UE::UAF
 {
 struct FInjectionInfo;
 }
 
-namespace UE::AnimNext::UncookedOnly
+namespace UE::UAF::UncookedOnly
 {
 struct FUtils;
 }
@@ -30,6 +30,6 @@ public:
 
 private:
 	friend class UAnimNextSharedVariablesFactory;
-	friend struct UE::AnimNext::UncookedOnly::FUtils;
-	friend UE::AnimNext::FInjectionInfo;
+	friend struct UE::UAF::UncookedOnly::FUtils;
+	friend UE::UAF::FInjectionInfo;
 };
diff --git a/Engine/Plugins/Experimental/AnimNext/Source/AnimNext/Public/Variables/UAFInstanceVariableData.h b/Engine/Plugins/Experimental/AnimNext/Source/AnimNext/Public/Variables/UAFInstanceVariableData.h
--- a/Engine/Plugins/Experimental/AnimNext/Source/AnimNext/Public/Variables/UAFInstanceVariableData.h
+++ b/Engine/Plugins/Experimental/AnimNext/Source/AnimNext/Public/Variables/UAFInstanceVariableData.h
@@ -21,7 +21,7 @@ struct FAnimNextRigVMAssetStructData;
 struct FUAFInstanceVariableContainer;
 struct FUAFRigVMComponent;
 
-namespace UE::AnimNext
+namespace UE::UAF
 {
 	struct FVariableOverrides;
 	struct FVariableOverridesCollection;
@@ -45,12 +45,12 @@ private:
 	friend FAnimNextModuleInstance;
 	friend FAnimNextGraphInstance;
 	friend FUAFInstanceVariableDataProxy;
-	friend UE::AnimNext::FAnimNextTrace;
-	friend UE::AnimNext::FInstanceTaskContext;
+	friend UE::UAF::FAnimNextTrace;
+	friend UE::UAF::FInstanceTaskContext;
 	friend FUAFRigVMComponent;
 
 	// Initialize from a host's asset and its references
-	void Initialize(const FUAFAssetInstance& InHostInstance, const TSharedPtr<const UE::AnimNext::FVariableOverridesCollection>& InOverrides);
+	void Initialize(const FUAFAssetInstance& InHostInstance, const TSharedPtr<const UE::UAF::FVariableOverridesCollection>& InOverrides);
 
 	// Reset to the default state
 	void Reset();
diff --git a/Engine/Plugins/Experimental/AnimNext/Source/AnimNextEditor/Internal/Common/AnimNextAssetItemDetails.cpp b/Engine/Plugins/Experimental/AnimNext/Source/AnimNextEditor/Internal/Common/AnimNextAssetItemDetails.cpp
--- a/Engine/Plugins/Experimental/AnimNext/Source/AnimNextEditor/Internal/Common/AnimNextAssetItemDetails.cpp
+++ b/Engine/Plugins/Experimental/AnimNext/Source/AnimNextEditor/Internal/Common/AnimNextAssetItemDetails.cpp
@@ -30,7 +30,7 @@
 
 #define LOCTEXT_NAMESPACE "FAnimNextGraphItemDetails"
 
-namespace UE::AnimNext::Editor
+namespace UE::UAF::Editor
 {
 
 const FSlateBrush* FAnimNextAssetItemDetails::GetItemIcon(const FWorkspaceOutlinerItemExport& Export) const
diff --git a/Engine/Plugins/Experimental/AnimNext/Source/AnimNextEditor/Internal/Common/AnimNextAssetItemDetails.h b/Engine/Plugins/Experimental/AnimNext/Source/AnimNextEditor/Internal/Common/AnimNextAssetItemDetails.h
--- a/Engine/Plugins/Experimental/AnimNext/Source/AnimNextEditor/Internal/Common/AnimNextAssetItemDetails.h
+++ b/Engine/Plugins/Experimental/AnimNext/Source/AnimNextEditor/Internal/Common/AnimNextAssetItemDetails.h
@@ -7,7 +7,7 @@
 class UAnimNextRigVMAssetEditorData;
 enum class EAnimNextEditorDataNotifType : uint8;
 
-namespace UE::AnimNext::Editor
+namespace UE::UAF::Editor
 {
 
 class FAnimNextAssetItemDetails : public UE::Workspace::IWorkspaceOutlinerItemDetails
diff --git a/Engine/Plugins/Experimental/AnimNext/Source/AnimNextEditor/Internal/Common/SCategoryTableRow.h b/Engine/Plugins/Experimental/AnimNext/Source/AnimNextEditor/Internal/Common/SCategoryTableRow.h
--- a/Engine/Plugins/Experimental/AnimNext/Source/AnimNextEditor/Internal/Common/SCategoryTableRow.h
+++ b/Engine/Plugins/Experimental/AnimNext/Source/AnimNextEditor/Internal/Common/SCategoryTableRow.h
@@ -9,7 +9,7 @@
 #include "Widgets/Views/SExpanderArrow.h"
 #include "Widgets/Views/STableRow.h"
 
-namespace UE::AnimNext::Editor
+namespace UE::UAF::Editor
 {
 
 // An entry category displayed in a tree view
@@ -101,4 +101,4 @@ private:
 	TSharedPtr<SBorder> ContentBorder;
 };
 
-} // end namespace UE::AnimNext::Editor
+} // end namespace UE::UAF::Editor
diff --git a/Engine/Plugins/Experimental/AnimNext/Source/AnimNextEditor/Internal/Common/SRigVMFunctionPicker.h b/Engine/Plugins/Experimental/AnimNext/Source/AnimNextEditor/Internal/Common/SRigVMFunctionPicker.h
--- a/Engine/Plugins/Experimental/AnimNext/Source/AnimNextEditor/Internal/Common/SRigVMFunctionPicker.h
+++ b/Engine/Plugins/Experimental/AnimNext/Source/AnimNextEditor/Internal/Common/SRigVMFunctionPicker.h
@@ -11,12 +11,12 @@ class UAnimNextRigVMAsset;
 class FTextFilterExpressionEvaluator;
 class SSearchBox;
 
-namespace UE::AnimNext::Editor
+namespace UE::UAF::Editor
 {
 	class SRigVMFunctionRowWidget;
 }
 
-namespace UE::AnimNext::Editor
+namespace UE::UAF::Editor
 {
 
 // Delegate called when the user picks a function from the dropdown menu
@@ -58,7 +58,7 @@ class SRigVMFunctionPicker : public SCompoundWidget
 	ANIMNEXTEDITOR_API void Construct(const FArguments& InArgs);
 
 private:
-	friend UE::AnimNext::Editor::SRigVMFunctionRowWidget;
+	friend UE::UAF::Editor::SRigVMFunctionRowWidget;
 
 	// Requests a refresh of the dropdown menu contents
 	void RequestRefreshEntries();
diff --git a/Engine/Plugins/Experimental/AnimNext/Source/AnimNextEditor/Private/AnimNextEditorModule.cpp b/Engine/Plugins/Experimental/AnimNext/Source/AnimNextEditor/Private/AnimNextEditorModule.cpp
--- a/Engine/Plugins/Experimental/AnimNext/Source/AnimNextEditor/Private/AnimNextEditorModule.cpp
+++ b/Engine/Plugins/Experimental/AnimNext/Source/AnimNextEditor/Private/AnimNextEditorModule.cpp
@@ -83,13 +83,13 @@ namespace
 {
 	FRewindDebuggerAnimNext GRewindDebuggerAnimNext;
 #if ANIMNEXT_TRACE_ENABLED
-	UE::AnimNextEditor::FAnimNextModuleTrackCreator GAnimNextModulesTrackCreator;
+	UE::UAF::Editor::FAnimNextModuleTrackCreator GAnimNextModulesTrackCreator;
 #endif
 	FAnimNextTraceModule GAnimNextTraceModule;
 }
 #endif
 
-namespace UE::AnimNext::Editor
+namespace UE::UAF::Editor
 {
 
 void FAnimNextEditorModule::StartupModule()
@@ -212,7 +212,7 @@ void FAnimNextEditorModule::StartupModule()
 								SearchReference = FAnimNextSoftVariableReference(SharedVariableNode->GetVariableName(), Object);
 								if (SharedVariableNode->Type == EAnimNextSharedVariablesType::Asset)
 								{
-									const UAnimNextRigVMAssetEditorData* EditorData = UE::AnimNext::UncookedOnly::FUtils::GetEditorData<const UAnimNextRigVMAssetEditorData, const UAnimNextSharedVariables>(Asset);						
+									const UAnimNextRigVMAssetEditorData* EditorData = UE::UAF::UncookedOnly::FUtils::GetEditorData<const UAnimNextRigVMAssetEditorData, const UAnimNextSharedVariables>(Asset);						
 									const UAnimNextVariableEntry* Entry = Cast<UAnimNextVariableEntry>(EditorData->FindEntry(SharedVariableNode->GetVariableName()));
 									SearchType = Entry->GetType();
 								}
@@ -666,7 +666,7 @@ void FAnimNextEditorModule::CollectGraphMenuActions(const TWeakPtr<UE::Workspace
 				continue;
 			}
 
-			UE::AnimNext::Editor::FUtils::AddSchemaRigUnitAction(URigVMUnitNode::StaticClass(), Struct, Function, InContextMenuBuilder);
+			UE::UAF::Editor::FUtils::AddSchemaRigUnitAction(URigVMUnitNode::StaticClass(), Struct, Function, InContextMenuBuilder);
 		}
 
 		for (const FRigVMDispatchFactory* Factory : FRigVMRegistry::Get().GetFactories())
@@ -721,10 +721,10 @@ void FAnimNextEditorModule::CollectGraphMenuActions(const TWeakPtr<UE::Workspace
 		}
 
 		TMap<FAssetData, FRigVMGraphFunctionHeaderArray> FunctionExports;
-		AnimNext::UncookedOnly::FUtils::GetExportedFunctionsFromAssetRegistry(UE::AnimNext::AnimNextPublicGraphFunctionsExportsRegistryTag, FunctionExports);
+		UncookedOnly::FUtils::GetExportedFunctionsFromAssetRegistry(UE::UAF::AnimNextPublicGraphFunctionsExportsRegistryTag, FunctionExports);
 		// TODO: Ideally we can filter functions by schema or execute context, but right now we dont expose the schema and function execute contexts are
 		// all FRigVMExecuteContext, rather than the 'most derived' context in the function.
-		//	AnimNext::UncookedOnly::FUtils::GetExportedFunctionsFromAssetRegistry(UE::AnimNext::ControlRigAssetPublicGraphFunctionsExportsRegistryTag, FunctionExports);
+		//	UncookedOnly::FUtils::GetExportedFunctionsFromAssetRegistry(UE::UAF::ControlRigAssetPublicGraphFunctionsExportsRegistryTag, FunctionExports);
 
 		for (const auto& Export : FunctionExports.Array())
 		{
@@ -790,7 +790,7 @@ void FAnimNextEditorModule::CollectGraphMenuActions(const TWeakPtr<UE::Workspace
 		// Promotion to variable
 		if (InContextMenuBuilder.FromPin != nullptr)
 		{
-			const FAnimNextParamType ParamType = UE::AnimNext::UncookedOnly::FUtils::GetParamTypeFromPinType(InContextMenuBuilder.FromPin->PinType);
+			const FAnimNextParamType ParamType = UE::UAF::UncookedOnly::FUtils::GetParamTypeFromPinType(InContextMenuBuilder.FromPin->PinType);
 			if (ParamType.IsValid())
 			{
 				InContextMenuBuilder.AddAction(MakeShared<FAnimNextSchemaAction_PromoteToVariable>());
@@ -1312,7 +1312,7 @@ void FAnimNextEditorModule::RegisterWorkspaceDocumentTypes(Workspace::IWorkspace
 							{
 								if (WeakEditorData.IsValid() && WeakEditorData->GetLocalFunctionLibrary() == WeakEdGraph->GetModel())
 								{
-									GraphName = UE::AnimNext::UncookedOnly::FUtils::GetFunctionLibraryDisplayName();
+									GraphName = UE::UAF::UncookedOnly::FUtils::GetFunctionLibraryDisplayName();
 								}
 								else if(UAnimNextRigVMAssetEntry* Entry = WeakEdGraph->GetTypedOuter<UAnimNextRigVMAssetEntry>())
 								{
@@ -1458,6 +1458,6 @@ const FAssetCompilationHandlerFactoryDelegate* FAnimNextEditorModule::FindAssetC
 
 }
 
-IMPLEMENT_MODULE(UE::AnimNext::Editor::FAnimNextEditorModule, AnimNextEditor);
+IMPLEMENT_MODULE(UE::UAF::Editor::FAnimNextEditorModule, AnimNextEditor);
 
 #undef LOCTEXT_NAMESPACE
diff --git a/Engine/Plugins/Experimental/AnimNext/Source/AnimNextEditor/Private/AnimNextEditorModule.h b/Engine/Plugins/Experimental/AnimNext/Source/AnimNextEditor/Private/AnimNextEditorModule.h
--- a/Engine/Plugins/Experimental/AnimNext/Source/AnimNextEditor/Private/AnimNextEditorModule.h
+++ b/Engine/Plugins/Experimental/AnimNext/Source/AnimNextEditor/Private/AnimNextEditorModule.h
@@ -11,7 +11,7 @@ class FAnimNextGraphPanelNodeFactory;
 class UAnimNextWorkspaceEditorMode;
 class FAnimNextGraphPanelPinFactory;
 
-namespace UE::AnimNext::Editor
+namespace UE::UAF::Editor
 {
 	class FParamNamePropertyTypeIdentifier;
 	class FLocatorContext;
@@ -22,7 +22,7 @@ namespace UE::Workspace
 	class IWorkspaceEditorModule;
 }
 
-namespace UE::AnimNext::Editor
+namespace UE::UAF::Editor
 {
 
 class FAnimNextEditorModule : public IAnimNextEditorModule
diff --git a/Engine/Plugins/Experimental/AnimNext/Source/AnimNextEditor/Private/AssetCompilationHandler.cpp b/Engine/Plugins/Experimental/AnimNext/Source/AnimNextEditor/Private/AssetCompilationHandler.cpp
--- a/Engine/Plugins/Experimental/AnimNext/Source/AnimNextEditor/Private/AssetCompilationHandler.cpp
+++ b/Engine/Plugins/Experimental/AnimNext/Source/AnimNextEditor/Private/AssetCompilationHandler.cpp
@@ -6,13 +6,13 @@
 #include "AnimNextScopedCompileJob.h"
 #include "UncookedOnlyUtils.h"
 
-namespace UE::AnimNext::Editor
+namespace UE::UAF::Editor
 {
 
 FAssetCompilationHandler::FAssetCompilationHandler(UObject* InAsset)
 {
 	UAnimNextRigVMAsset* AnimNextRigVMAsset = CastChecked<UAnimNextRigVMAsset>(InAsset);
-	UAnimNextRigVMAssetEditorData* EditorData = UE::AnimNext::UncookedOnly::FUtils::GetEditorData<UAnimNextRigVMAssetEditorData>(AnimNextRigVMAsset);
+	UAnimNextRigVMAssetEditorData* EditorData = UE::UAF::UncookedOnly::FUtils::GetEditorData<UAnimNextRigVMAssetEditorData>(AnimNextRigVMAsset);
 
 	NumErrors = EditorData->bErrorsDuringCompilation ? 1 : 0;
 	NumWarnings = EditorData->bWarningsDuringCompilation ? 1 : 0;
@@ -21,7 +21,7 @@ FAssetCompilationHandler::FAssetCompilationHandler(UObject* InAsset)
 void FAssetCompilationHandler::Compile(TSharedRef<Workspace::IWorkspaceEditor> InWorkspaceEditor, UObject* InAsset)
 {
 	UAnimNextRigVMAsset* AnimNextRigVMAsset = CastChecked<UAnimNextRigVMAsset>(InAsset);
-	UAnimNextRigVMAssetEditorData* EditorData = UE::AnimNext::UncookedOnly::FUtils::GetEditorData<UAnimNextRigVMAssetEditorData>(AnimNextRigVMAsset);
+	UAnimNextRigVMAssetEditorData* EditorData = UE::UAF::UncookedOnly::FUtils::GetEditorData<UAnimNextRigVMAssetEditorData>(AnimNextRigVMAsset);
 
 	UncookedOnly::FScopedCompileJob CompileJob(InAsset);
 	EditorData->RecompileVM();
@@ -35,21 +35,21 @@ void FAssetCompilationHandler::Compile(TSharedRef<Workspace::IWorkspaceEditor> I
 void FAssetCompilationHandler::SetAutoCompile(TSharedRef<Workspace::IWorkspaceEditor> InWorkspaceEditor, UObject* InAsset, bool bInAutoCompile)
 {
 	UAnimNextRigVMAsset* AnimNextRigVMAsset = CastChecked<UAnimNextRigVMAsset>(InAsset);
-	UAnimNextRigVMAssetEditorData* EditorData = UE::AnimNext::UncookedOnly::FUtils::GetEditorData<UAnimNextRigVMAssetEditorData>(AnimNextRigVMAsset);
+	UAnimNextRigVMAssetEditorData* EditorData = UE::UAF::UncookedOnly::FUtils::GetEditorData<UAnimNextRigVMAssetEditorData>(AnimNextRigVMAsset);
 	EditorData->SetAutoVMRecompile(bInAutoCompile);
 }
 
 bool FAssetCompilationHandler::GetAutoCompile(TSharedRef<Workspace::IWorkspaceEditor> InWorkspaceEditor, const UObject* InAsset) const
 {
 	const UAnimNextRigVMAsset* AnimNextRigVMAsset = CastChecked<UAnimNextRigVMAsset>(InAsset);
-	UAnimNextRigVMAssetEditorData* EditorData = UE::AnimNext::UncookedOnly::FUtils::GetEditorData<UAnimNextRigVMAssetEditorData>(AnimNextRigVMAsset);
+	UAnimNextRigVMAssetEditorData* EditorData = UE::UAF::UncookedOnly::FUtils::GetEditorData<UAnimNextRigVMAssetEditorData>(AnimNextRigVMAsset);
 	return EditorData->GetAutoVMRecompile();
 }
 
 ECompileStatus FAssetCompilationHandler::GetCompileStatus(TSharedRef<Workspace::IWorkspaceEditor> InWorkspaceEditor, const UObject* InAsset) const
 {
 	const UAnimNextRigVMAsset* AnimNextRigVMAsset = CastChecked<UAnimNextRigVMAsset>(InAsset);
-	UAnimNextRigVMAssetEditorData* EditorData = UE::AnimNext::UncookedOnly::FUtils::GetEditorData<UAnimNextRigVMAssetEditorData>(AnimNextRigVMAsset);
+	UAnimNextRigVMAssetEditorData* EditorData = UE::UAF::UncookedOnly::FUtils::GetEditorData<UAnimNextRigVMAssetEditorData>(AnimNextRigVMAsset);
 	if(EditorData->bErrorsDuringCompilation || NumErrors > 0)
 	{
 		return ECompileStatus::Error;
diff --git a/Engine/Plugins/Experimental/AnimNext/Source/AnimNextEditor/Private/Common/AnimNextAssetFindReplaceVariables.cpp b/Engine/Plugins/Experimental/AnimNext/Source/AnimNextEditor/Private/Common/AnimNextAssetFindReplaceVariables.cpp
--- a/Engine/Plugins/Experimental/AnimNext/Source/AnimNextEditor/Private/Common/AnimNextAssetFindReplaceVariables.cpp
+++ b/Engine/Plugins/Experimental/AnimNext/Source/AnimNextEditor/Private/Common/AnimNextAssetFindReplaceVariables.cpp
@@ -24,7 +24,7 @@
 FString UAnimNextAssetFindReplaceVariables::GetFindResultStringFromAssetData(const FAssetData& InAssetData) const
 {
 	FAnimNextAssetRegistryExports AnimNextExports;
-	UE::AnimNext::UncookedOnly::FUtils::GetExportsOfTypeForAsset<FAnimNextVariableReferenceData>(InAssetData, AnimNextExports);
+	UE::UAF::UncookedOnly::FUtils::GetExportsOfTypeForAsset<FAnimNextVariableReferenceData>(InAssetData, AnimNextExports);
 
 	TArray<FString> ReferenceNames;
 	
@@ -89,7 +89,7 @@ bool UAnimNextAssetFindReplaceVariables::ShouldFilterOutAsset(const FAssetData&
 	}
 	
 	FAnimNextAssetRegistryExports AnimNextExports;
-	UE::AnimNext::UncookedOnly::FUtils::GetExportsOfTypeForAsset<FAnimNextVariableReferenceData>(InAssetData, AnimNextExports);
+	UE::UAF::UncookedOnly::FUtils::GetExportsOfTypeForAsset<FAnimNextVariableReferenceData>(InAssetData, AnimNextExports);
 	
 	for (const FAnimNextExport& Export : AnimNextExports.Exports)
 	{
@@ -106,9 +106,9 @@ void UAnimNextAssetFindReplaceVariables::ReplaceInAsset(const FAssetData& InAsse
 {
 	if (UAnimNextRigVMAsset* RigVMAsset = Cast<UAnimNextRigVMAsset>(InAssetData.GetAsset()))
 	{
-		if (UAnimNextRigVMAssetEditorData* EditorData = UE::AnimNext::UncookedOnly::FUtils::GetEditorData<UAnimNextRigVMAssetEditorData>(RigVMAsset))
+		if (UAnimNextRigVMAssetEditorData* EditorData = UE::UAF::UncookedOnly::FUtils::GetEditorData<UAnimNextRigVMAssetEditorData>(RigVMAsset))
 		{
-			UE::AnimNext::UncookedOnly::FUtils::ReplaceVariableReferences(EditorData, SearchReference, ReplaceReference);
+			UE::UAF::UncookedOnly::FUtils::ReplaceVariableReferences(EditorData, SearchReference, ReplaceReference);
 		}		
 	}	
 }
@@ -117,10 +117,10 @@ void UAnimNextAssetFindReplaceVariables::RemoveInAsset(const FAssetData& InAsset
 {	
 	if (UAnimNextRigVMAsset* RigVMAsset = Cast<UAnimNextRigVMAsset>(InAssetData.GetAsset()))
 	{
-		if (UAnimNextRigVMAssetEditorData* EditorData = UE::AnimNext::UncookedOnly::FUtils::GetEditorData<UAnimNextRigVMAssetEditorData>(RigVMAsset))
+		if (UAnimNextRigVMAssetEditorData* EditorData = UE::UAF::UncookedOnly::FUtils::GetEditorData<UAnimNextRigVMAssetEditorData>(RigVMAsset))
 		{
 			// Replace with None reference (removes nodes)
-			UE::AnimNext::UncookedOnly::FUtils::ReplaceVariableReferences(EditorData, SearchReference, FAnimNextSoftVariableReference());
+			UE::UAF::UncookedOnly::FUtils::ReplaceVariableReferences(EditorData, SearchReference, FAnimNextSoftVariableReference());
 		}		
 	}
 }
@@ -240,7 +240,7 @@ FSlateIcon UAnimNextAssetFindReplaceVariables::GetIconFromSearchScope(ESearchSco
 
 TSharedRef<SWidget> UAnimNextAssetFindReplaceVariables::MakeFindReplaceWidget()
 {	
-	using namespace UE::AnimNext::Editor;
+	using namespace UE::UAF::Editor;
 
 
 	auto VariableScopeFilter = FOnFilterVariable::CreateLambda([this](const FAnimNextSoftVariableReference& InVariableReference)
@@ -262,7 +262,7 @@ TSharedRef<SWidget> UAnimNextAssetFindReplaceVariables::MakeFindReplaceWidget()
       			{
       				if (InAsset)
       				{
-      					const UAnimNextRigVMAssetEditorData* EditorData = UE::AnimNext::UncookedOnly::FUtils::GetEditorData<const UAnimNextRigVMAssetEditorData,const UAnimNextRigVMAsset>(InAsset);
+      					const UAnimNextRigVMAssetEditorData* EditorData = UE::UAF::UncookedOnly::FUtils::GetEditorData<const UAnimNextRigVMAssetEditorData,const UAnimNextRigVMAsset>(InAsset);
       					EditorData->ForEachEntryOfType<UAnimNextSharedVariablesEntry>([&AssetsInExportChain](const UAnimNextSharedVariablesEntry* VariablesEntry)
       					{
       						// [TODO] struct, currently doesn't have any exports so cannot filter against them
@@ -349,7 +349,7 @@ TSharedRef<SWidget> UAnimNextAssetFindReplaceVariables::MakeFindReplaceWidget()
 			if(SearchType.IsValid())
 			{
 				// [TODO] IsCompatible or IsCompatibleWithDataLoss (e.g. UScriptStruct float entry to double RigVM entry fails IsCompatible)
-				if(!UE::AnimNext::FParamUtils::GetCompatibility(SearchType, InParamType).IsCompatible())
+				if(!UE::UAF::FParamUtils::GetCompatibility(SearchType, InParamType).IsCompatible())
 				{
 					return EFilterVariableResult::Exclude;
 				}
diff --git a/Engine/Plugins/Experimental/AnimNext/Source/AnimNextEditor/Private/Common/AnimNextAssetFindReplaceVariables.h b/Engine/Plugins/Experimental/AnimNext/Source/AnimNextEditor/Private/Common/AnimNextAssetFindReplaceVariables.h
--- a/Engine/Plugins/Experimental/AnimNext/Source/AnimNextEditor/Private/Common/AnimNextAssetFindReplaceVariables.h
+++ b/Engine/Plugins/Experimental/AnimNext/Source/AnimNextEditor/Private/Common/AnimNextAssetFindReplaceVariables.h
@@ -12,7 +12,7 @@
 
 class UAnimNextVariableEntry;
 
-namespace UE::AnimNext::Editor
+namespace UE::UAF::Editor
 {
 class SVariablePickerCombo;
 }
@@ -63,8 +63,8 @@ protected:
 	ECheckBoxState IsCurrentSearchScope(ESearchScope InSearchScope) const;
 	void OnWorkspaceDocumentChanged(const UE::Workspace::FWorkspaceDocument& InDocument);
 protected:
-	TSharedPtr<UE::AnimNext::Editor::SVariablePickerCombo> SearchVariableComboBox;	
-	TSharedPtr<UE::AnimNext::Editor::SVariablePickerCombo> ReplaceVariableComboBox;
+	TSharedPtr<UE::UAF::Editor::SVariablePickerCombo> SearchVariableComboBox;	
+	TSharedPtr<UE::UAF::Editor::SVariablePickerCombo> ReplaceVariableComboBox;
 
 	FAnimNextSoftVariableReference SearchReference;
 	FAnimNextParamType SearchType;
diff --git a/Engine/Plugins/Experimental/AnimNext/Source/AnimNextEditor/Private/Common/AnimNextCollapseNodeItemDetails.cpp b/Engine/Plugins/Experimental/AnimNext/Source/AnimNextEditor/Private/Common/AnimNextCollapseNodeItemDetails.cpp
--- a/Engine/Plugins/Experimental/AnimNext/Source/AnimNextEditor/Private/Common/AnimNextCollapseNodeItemDetails.cpp
+++ b/Engine/Plugins/Experimental/AnimNext/Source/AnimNextEditor/Private/Common/AnimNextCollapseNodeItemDetails.cpp
@@ -17,7 +17,7 @@
 
 #define LOCTEXT_NAMESPACE "FAnimNextCollapseNodeItemDetails"
 
-namespace UE::AnimNext::Editor
+namespace UE::UAF::Editor
 {
 
 bool FAnimNextCollapseNodeItemDetails::HandleDoubleClick(const FToolMenuContext& ToolMenuContext) const
@@ -190,6 +190,6 @@ void FAnimNextCollapseNodeItemDetails::UnregisterToolMenuExtensions()
 }
 
 
-} // UE::AnimNext::Editor
+} // UE::UAF::Editor
 
 #undef LOCTEXT_NAMESPACE // "FAnimNextCollapseNodeItemDetails"
diff --git a/Engine/Plugins/Experimental/AnimNext/Source/AnimNextEditor/Private/Common/AnimNextCompilerResultsTabSummoner.cpp b/Engine/Plugins/Experimental/AnimNext/Source/AnimNextEditor/Private/Common/AnimNextCompilerResultsTabSummoner.cpp
--- a/Engine/Plugins/Experimental/AnimNext/Source/AnimNextEditor/Private/Common/AnimNextCompilerResultsTabSummoner.cpp
+++ b/Engine/Plugins/Experimental/AnimNext/Source/AnimNextEditor/Private/Common/AnimNextCompilerResultsTabSummoner.cpp
@@ -12,7 +12,7 @@
 
 #define LOCTEXT_NAMESPACE "WorkspaceTabSummoner"
 
-namespace UE::AnimNext::Editor
+namespace UE::UAF::Editor
 {
 
 // ***************************************************************************
@@ -80,7 +80,7 @@ void SAnimNextCompilerResultsWidget::HandleMessageTokenClicked(const TSharedRef<
 // ***************************************************************************
 
 FAnimNextCompilerResultsTabSummoner::FAnimNextCompilerResultsTabSummoner(TSharedPtr<UE::Workspace::IWorkspaceEditor> InHostingApp)
-	: FWorkflowTabFactory(UE::AnimNext::Editor::CompilerResultsTabName, StaticCastSharedPtr<FAssetEditorToolkit>(InHostingApp))
+	: FWorkflowTabFactory(UE::UAF::Editor::CompilerResultsTabName, StaticCastSharedPtr<FAssetEditorToolkit>(InHostingApp))
 {
 	TabLabel = LOCTEXT("AnimNExtCompilerResultsTabLabel", "Compiler Results");
 	TabIcon = FSlateIcon("EditorStyle", "LevelEditor.Tabs.Outliner");
@@ -101,6 +101,6 @@ FText FAnimNextCompilerResultsTabSummoner::GetTabToolTipText(const FWorkflowTabS
 	return ViewMenuTooltip;
 }
 
-} // end namespace UE::AnimNext::Editor
+} // end namespace UE::UAF::Editor
 
 #undef LOCTEXT_NAMESPACE
diff --git a/Engine/Plugins/Experimental/AnimNext/Source/AnimNextEditor/Private/Common/AnimNextCompilerResultsTabSummoner.h b/Engine/Plugins/Experimental/AnimNext/Source/AnimNextEditor/Private/Common/AnimNextCompilerResultsTabSummoner.h
--- a/Engine/Plugins/Experimental/AnimNext/Source/AnimNextEditor/Private/Common/AnimNextCompilerResultsTabSummoner.h
+++ b/Engine/Plugins/Experimental/AnimNext/Source/AnimNextEditor/Private/Common/AnimNextCompilerResultsTabSummoner.h
@@ -10,7 +10,7 @@ namespace UE::Workspace
 	class IWorkspaceEditor;
 }
 
-namespace UE::AnimNext::Editor
+namespace UE::UAF::Editor
 {
 class SAnimNextCompilerResultsWidget : public SCompoundWidget
 {
@@ -50,4 +50,4 @@ private:
 	TSharedPtr<class SAnimNextCompilerResultsWidget> AnimNextCompilerResultsWidget;
 };
 
-} // end namespace UE::AnimNext::Editor
+} // end namespace UE::UAF::Editor
diff --git a/Engine/Plugins/Experimental/AnimNext/Source/AnimNextEditor/Private/Common/AnimNextFindTabSummoner.cpp b/Engine/Plugins/Experimental/AnimNext/Source/AnimNextEditor/Private/Common/AnimNextFindTabSummoner.cpp
--- a/Engine/Plugins/Experimental/AnimNext/Source/AnimNextEditor/Private/Common/AnimNextFindTabSummoner.cpp
+++ b/Engine/Plugins/Experimental/AnimNext/Source/AnimNextEditor/Private/Common/AnimNextFindTabSummoner.cpp
@@ -13,11 +13,11 @@
 
 #define LOCTEXT_NAMESPACE "WorkspaceTabSummoner"
 
-namespace UE::AnimNext::Editor
+namespace UE::UAF::Editor
 {
 
 FAnimNextFindTabSummoner::FAnimNextFindTabSummoner(TSharedPtr<UE::Workspace::IWorkspaceEditor> InHostingApp)
-	: FWorkflowTabFactory(UE::AnimNext::Editor::FindTabName, StaticCastSharedPtr<FAssetEditorToolkit>(InHostingApp))
+	: FWorkflowTabFactory(UE::UAF::Editor::FindTabName, StaticCastSharedPtr<FAssetEditorToolkit>(InHostingApp))
 {
 	TabLabel = LOCTEXT("UAFFindTabLabel", "Find");
 	TabIcon = FSlateIcon(FAppStyle::GetAppStyleSetName(), "Kismet.Tabs.FindResults");
@@ -39,7 +39,7 @@ FText FAnimNextFindTabSummoner::GetTabToolTipText(const FWorkflowTabSpawnInfo& I
 }
 
 FAnimNextFindAndReplaceTabSummoner::FAnimNextFindAndReplaceTabSummoner(TSharedPtr<UE::Workspace::IWorkspaceEditor> InHostingApp) 
-	: FWorkflowTabFactory(UE::AnimNext::Editor::FindAndReplaceTabName, StaticCastSharedPtr<FAssetEditorToolkit>(InHostingApp))
+	: FWorkflowTabFactory(UE::UAF::Editor::FindAndReplaceTabName, StaticCastSharedPtr<FAssetEditorToolkit>(InHostingApp))
 {
 	TabLabel = LOCTEXT("AnimNextFindAndReplaceTabLabel", "Find and Replace");
 	TabIcon = FSlateIcon(FAppStyle::GetAppStyleSetName(), "Kismet.Tabs.FindResults");
@@ -69,6 +69,6 @@ FText FAnimNextFindAndReplaceTabSummoner::GetTabToolTipText(const FWorkflowTabSp
 {
 	return ViewMenuTooltip;
 }
-} // end namespace UE::AnimNext::Editor
+} // end namespace UE::UAF::Editor
 
 #undef LOCTEXT_NAMESPACE
diff --git a/Engine/Plugins/Experimental/AnimNext/Source/AnimNextEditor/Private/Common/AnimNextFindTabSummoner.h b/Engine/Plugins/Experimental/AnimNext/Source/AnimNextEditor/Private/Common/AnimNextFindTabSummoner.h
--- a/Engine/Plugins/Experimental/AnimNext/Source/AnimNextEditor/Private/Common/AnimNextFindTabSummoner.h
+++ b/Engine/Plugins/Experimental/AnimNext/Source/AnimNextEditor/Private/Common/AnimNextFindTabSummoner.h
@@ -10,14 +10,14 @@ namespace UE::Workspace
 	class IWorkspaceEditor;
 }	// end namespace UE::Workspace
 
-namespace UE::AnimNext::Editor
+namespace UE::UAF::Editor
 {
 	class SFindInAnimNextRigVMAsset;
-}	// end namespace UE::AnimNext::Editor
+}	// end namespace UE::UAF::Editor
 
 class IAnimAssetFindReplace;
 
-namespace UE::AnimNext::Editor
+namespace UE::UAF::Editor
 {
 struct FAnimNextFindTabSummoner : public FWorkflowTabFactory
 {
@@ -47,4 +47,4 @@ private:
 	TSharedPtr<IAnimAssetFindReplace> FindAndReplaceWidget;
 };
 
-} // end namespace UE::AnimNext::Editor
+} // end namespace UE::UAF::Editor
diff --git a/Engine/Plugins/Experimental/AnimNext/Source/AnimNextEditor/Private/Common/AnimNextFunctionItemDetails.cpp b/Engine/Plugins/Experimental/AnimNext/Source/AnimNextEditor/Private/Common/AnimNextFunctionItemDetails.cpp
--- a/Engine/Plugins/Experimental/AnimNext/Source/AnimNextEditor/Private/Common/AnimNextFunctionItemDetails.cpp
+++ b/Engine/Plugins/Experimental/AnimNext/Source/AnimNextEditor/Private/Common/AnimNextFunctionItemDetails.cpp
@@ -17,7 +17,7 @@
 
 #define LOCTEXT_NAMESPACE "FAnimNextFunctionItemDetails"
 
-namespace UE::AnimNext::Editor
+namespace UE::UAF::Editor
 {
 
 bool FAnimNextFunctionItemDetails::HandleDoubleClick(const FToolMenuContext& ToolMenuContext) const
@@ -224,6 +224,6 @@ void FAnimNextFunctionItemDetails::UnregisterToolMenuExtensions()
 {
 }
 
-} // UE::AnimNext::Editor
+} // UE::UAF::Editor
 
 #undef LOCTEXT_NAMESPACE // "FAnimNextFunctionItemDetails"
diff --git a/Engine/Plugins/Experimental/AnimNext/Source/AnimNextEditor/Private/Common/FindInAnimNextRigVMAsset.cpp b/Engine/Plugins/Experimental/AnimNext/Source/AnimNextEditor/Private/Common/FindInAnimNextRigVMAsset.cpp
--- a/Engine/Plugins/Experimental/AnimNext/Source/AnimNextEditor/Private/Common/FindInAnimNextRigVMAsset.cpp
+++ b/Engine/Plugins/Experimental/AnimNext/Source/AnimNextEditor/Private/Common/FindInAnimNextRigVMAsset.cpp
@@ -13,7 +13,7 @@
 
 #define LOCTEXT_NAMESPACE "FindInAnimNextRigVMAsset"
 
-namespace UE::AnimNext::Editor
+namespace UE::UAF::Editor
 {
 
 //////////////////////////////////////////////////////////////////////////
@@ -73,7 +73,7 @@ void SFindInAnimNextRigVMAsset::MatchTokens(const TArray<FString>& Tokens)
 		else if (UAnimNextRigVMAsset* AnimNextRigVMAsset = FocusedDocument.GetTypedObject<UAnimNextRigVMAsset>())
 		{
 			// Note: This is where we search subgraphs in say a module
-			if (UAnimNextRigVMAssetEditorData* AnimNextRigVMAssetEditorData = UE::AnimNext::UncookedOnly::FUtils::GetEditorData<UAnimNextRigVMAssetEditorData>(AnimNextRigVMAsset))
+			if (UAnimNextRigVMAssetEditorData* AnimNextRigVMAssetEditorData = UE::UAF::UncookedOnly::FUtils::GetEditorData<UAnimNextRigVMAssetEditorData>(AnimNextRigVMAsset))
 			{
 				TArray<UEdGraph*> ContainedGraphs = AnimNextRigVMAssetEditorData->GetAllEdGraphs();
 				for (const UEdGraph* ContainedGraph : ContainedGraphs)
@@ -104,7 +104,7 @@ bool SFindInAnimNextRigVMAsset::MatchTokensInNode(const UEdGraphNode* Node, cons
 							const FSoftObjectPath ObjectPath(RigVMPin->GetDefaultValue());
 							if (UAnimNextRigVMAsset* AnimNextRigVMAsset = Cast<UAnimNextRigVMAsset>(ObjectPath.ResolveObject()))
 							{
-								if (UAnimNextRigVMAssetEditorData* AnimNextRigVMAssetEditorData = UE::AnimNext::UncookedOnly::FUtils::GetEditorData<UAnimNextRigVMAssetEditorData>(AnimNextRigVMAsset))
+								if (UAnimNextRigVMAssetEditorData* AnimNextRigVMAssetEditorData = UE::UAF::UncookedOnly::FUtils::GetEditorData<UAnimNextRigVMAssetEditorData>(AnimNextRigVMAsset))
 								{
 									TArray<UEdGraph*> ContainedGraphs = AnimNextRigVMAssetEditorData->GetAllEdGraphs();
 									for (const UEdGraph* ContainedGraph : ContainedGraphs)
@@ -126,7 +126,7 @@ bool SFindInAnimNextRigVMAsset::MatchTokensInNode(const UEdGraphNode* Node, cons
 
 /////////////////////////////////////////////////////
 
-} // namespace UE::AnimNext::Editor
+} // namespace UE::UAF::Editor
 
 #undef LOCTEXT_NAMESPACE
 
diff --git a/Engine/Plugins/Experimental/AnimNext/Source/AnimNextEditor/Private/Common/GraphEditorSchemaActions.cpp b/Engine/Plugins/Experimental/AnimNext/Source/AnimNextEditor/Private/Common/GraphEditorSchemaActions.cpp
--- a/Engine/Plugins/Experimental/AnimNext/Source/AnimNextEditor/Private/Common/GraphEditorSchemaActions.cpp
+++ b/Engine/Plugins/Experimental/AnimNext/Source/AnimNextEditor/Private/Common/GraphEditorSchemaActions.cpp
@@ -50,7 +50,7 @@ UEdGraphNode* FAnimNextSchemaAction_RigUnit::PerformAction(UEdGraph* ParentGraph
 	
 	if (Host != nullptr && EdGraph != nullptr)
 	{
-		FName Name = UE::AnimNext::Editor::FUtils::ValidateName(Cast<UObject>(Host), StructTemplate->GetFName().ToString());
+		FName Name = UE::UAF::Editor::FUtils::ValidateName(Cast<UObject>(Host), StructTemplate->GetFName().ToString());
 		URigVMController* Controller = Host->GetRigVMClient()->GetController(ParentGraph);
 
 		Controller->OpenUndoBracket(FString::Printf(TEXT("Add '%s' Node"), *Name.ToString()));
@@ -117,7 +117,7 @@ UEdGraphNode* FAnimNextSchemaAction_DispatchFactory::PerformAction(UEdGraph* Par
 	const int32 NotationHash = (int32)GetTypeHash(Notation);
 	const FString TemplateName = TEXT("RigVMTemplate_") + FString::FromInt(NotationHash);
 
-	FName Name = UE::AnimNext::Editor::FUtils::ValidateName(Cast<UObject>(Host), Template->GetName().ToString());
+	FName Name = UE::UAF::Editor::FUtils::ValidateName(Cast<UObject>(Host), Template->GetName().ToString());
 	URigVMController* Controller = Host->GetRigVMClient()->GetController(ParentGraph);
 
 	Controller->OpenUndoBracket(FString::Printf(TEXT("Add '%s' Node"), *Name.ToString()));
@@ -187,7 +187,7 @@ FAnimNextSchemaAction_Variable::FAnimNextSchemaAction_Variable(FName InName, con
 
 	UpdateSearchData(MenuDesc, ToolTip, VariablesCategory, FText::GetEmpty());
 
-	FEdGraphPinType PinType = UE::AnimNext::UncookedOnly::FUtils::GetPinTypeFromParamType(InType);
+	FEdGraphPinType PinType = UE::UAF::UncookedOnly::FUtils::GetPinTypeFromParamType(InType);
 	VariableColor = GetDefault<UEdGraphSchema_K2>()->GetPinTypeColor(PinType);
 }
 
@@ -253,15 +253,15 @@ UEdGraphNode* FAnimNextSchemaAction_Variable::PerformAction(UEdGraph* ParentGrap
 			// Satisfy host dependency
 			if (UAnimNextRigVMAssetEditorData* OwnerEditorData = ParentGraph->GetTypedOuter<UAnimNextRigVMAssetEditorData>())
 			{
-				if (UAnimNextRigVMAssetEditorData* VariableOwnerEditorData = UE::AnimNext::UncookedOnly::FUtils::GetEditorData(VariableOwningRigVMAsset))
+				if (UAnimNextRigVMAssetEditorData* VariableOwnerEditorData = UE::UAF::UncookedOnly::FUtils::GetEditorData(VariableOwningRigVMAsset))
 				{
 					if (OwnerEditorData != VariableOwnerEditorData)
 					{
 						if (!OwnerEditorData->FindEntry(VariableOwningRigVMAsset->GetFName()))
 						{
-							UE::AnimNext::UncookedOnly::FScopedCompileJob CompilerResults(LOCTEXT("AddSharedAssetJobName", "Add Shared Asset dependency"), { ParentGraph->GetTypedOuter<UAnimNextRigVMAsset>() });
+							UE::UAF::UncookedOnly::FScopedCompileJob CompilerResults(LOCTEXT("AddSharedAssetJobName", "Add Shared Asset dependency"), { ParentGraph->GetTypedOuter<UAnimNextRigVMAsset>() });
 							
-							UAnimNextSharedVariables* VariableOwningSharedVariables = UE::AnimNext::UncookedOnly::FUtils::GetAsset<UAnimNextSharedVariables>(VariableOwnerEditorData);
+							UAnimNextSharedVariables* VariableOwningSharedVariables = UE::UAF::UncookedOnly::FUtils::GetAsset<UAnimNextSharedVariables>(VariableOwnerEditorData);
 							OwnerEditorData->AddSharedVariables(VariableOwningSharedVariables, true, false);
 							OwnerEditorData->RecompileVM();
 						}
@@ -424,7 +424,7 @@ UEdGraphNode* FAnimNextSchemaAction_Function::PerformAction(UEdGraph* ParentGrap
 
 	if (Host != nullptr && EdGraph != nullptr)
 	{
-		FName Name = UE::AnimNext::Editor::FUtils::ValidateName(Cast<UObject>(Host), ReferencedPublicFunctionHeader.Name.ToString());
+		FName Name = UE::UAF::Editor::FUtils::ValidateName(Cast<UObject>(Host), ReferencedPublicFunctionHeader.Name.ToString());
 		URigVMController* Controller = EdGraph->GetController();
 
 		Controller->OpenUndoBracket(FString::Printf(TEXT("Add '%s' Node"), *Name.ToString()));
@@ -472,8 +472,8 @@ UEdGraphNode* FAnimNextSchemaAction_PromoteToVariable::PerformAction(UEdGraph* P
 			{
 				if (UEdGraphPin* FromPin = FromPins[0])
 				{
-					UAnimNextRigVMAsset* RigVMAsset = UE::AnimNext::UncookedOnly::FUtils::GetAsset(EditorData);
-					const FAnimNextParamType ParamType = UE::AnimNext::UncookedOnly::FUtils::GetParamTypeFromPinType(FromPin->PinType);
+					UAnimNextRigVMAsset* RigVMAsset = UE::UAF::UncookedOnly::FUtils::GetAsset(EditorData);
+					const FAnimNextParamType ParamType = UE::UAF::UncookedOnly::FUtils::GetParamTypeFromPinType(FromPin->PinType);
 					if (ParamType.IsValid() && EditorData)
 					{
 						FScopedTransaction Transaction(LOCTEXT("PromoteToVariable", "Promote to Variable"));
@@ -482,7 +482,7 @@ UEdGraphNode* FAnimNextSchemaAction_PromoteToVariable::PerformAction(UEdGraph* P
 						BaseName.ReplaceCharInline(' ', '_');
 
 						const FName BaseFName = FName(*BaseName);
-						const FName NewVariableName = UE::AnimNext::UncookedOnly::FUtils::GetValidVariableName(EditorData, BaseFName);
+						const FName NewVariableName = UE::UAF::UncookedOnly::FUtils::GetValidVariableName(EditorData, BaseFName);
 						if (UAnimNextVariableEntry* VariableEntry = EditorData->AddVariable(NewVariableName, ParamType, FromPin->GetDefaultAsString()))
 						{
 							UAnimNextControllerBase* Controller = CastChecked<UAnimNextControllerBase>(Host->GetRigVMClient()->GetController(ParentGraph));
diff --git a/Engine/Plugins/Experimental/AnimNext/Source/AnimNextEditor/Private/Common/RigVMAssetViewMenuContext.h b/Engine/Plugins/Experimental/AnimNext/Source/AnimNextEditor/Private/Common/RigVMAssetViewMenuContext.h
--- a/Engine/Plugins/Experimental/AnimNext/Source/AnimNextEditor/Private/Common/RigVMAssetViewMenuContext.h
+++ b/Engine/Plugins/Experimental/AnimNext/Source/AnimNextEditor/Private/Common/RigVMAssetViewMenuContext.h
@@ -4,7 +4,7 @@
 
 #include "RigVMAssetViewMenuContext.generated.h"
 
-namespace UE::AnimNext::Editor
+namespace UE::UAF::Editor
 {
 	class SRigVMAssetView;
 }
@@ -14,8 +14,8 @@ class URigVMAssetViewMenuContext : public UObject
 {
 	GENERATED_BODY()
 
-	friend class UE::AnimNext::Editor::SRigVMAssetView;
+	friend class UE::UAF::Editor::SRigVMAssetView;
 	
 	// The RigVM asset view that we are editing
-	TWeakPtr<UE::AnimNext::Editor::SRigVMAssetView> RigVMAssetView;
+	TWeakPtr<UE::UAF::Editor::SRigVMAssetView> RigVMAssetView;
 };
\ No newline at end of file
diff --git a/Engine/Plugins/Experimental/AnimNext/Source/AnimNextEditor/Private/EditorCVars.cpp b/Engine/Plugins/Experimental/AnimNext/Source/AnimNextEditor/Private/EditorCVars.cpp
--- a/Engine/Plugins/Experimental/AnimNext/Source/AnimNextEditor/Private/EditorCVars.cpp
+++ b/Engine/Plugins/Experimental/AnimNext/Source/AnimNextEditor/Private/EditorCVars.cpp
@@ -2,7 +2,7 @@
 
 #include "EditorCVars.h"
 
-namespace UE::AnimNext::Editor::CVars
+namespace UE::UAF::Editor::CVars
 {
 
 TAutoConsoleVariable<bool> GUseWorkspaceEditor(
diff --git a/Engine/Plugins/Experimental/AnimNext/Source/AnimNextEditor/Private/EditorCVars.h b/Engine/Plugins/Experimental/AnimNext/Source/AnimNextEditor/Private/EditorCVars.h
--- a/Engine/Plugins/Experimental/AnimNext/Source/AnimNextEditor/Private/EditorCVars.h
+++ b/Engine/Plugins/Experimental/AnimNext/Source/AnimNextEditor/Private/EditorCVars.h
@@ -4,7 +4,7 @@
 
 #include "HAL/IConsoleManager.h"
 
-namespace UE::AnimNext::Editor::CVars
+namespace UE::UAF::Editor::CVars
 {
 
 extern TAutoConsoleVariable<bool> GUseWorkspaceEditor;
diff --git a/Engine/Plugins/Experimental/AnimNext/Source/AnimNextEditor/Private/EditorUtils.cpp b/Engine/Plugins/Experimental/AnimNext/Source/AnimNextEditor/Private/EditorUtils.cpp
--- a/Engine/Plugins/Experimental/AnimNext/Source/AnimNextEditor/Private/EditorUtils.cpp
+++ b/Engine/Plugins/Experimental/AnimNext/Source/AnimNextEditor/Private/EditorUtils.cpp
@@ -19,7 +19,7 @@
 
 #define LOCTEXT_NAMESPACE "AnimNextEditorUtils"
 
-namespace UE::AnimNext::Editor
+namespace UE::UAF::Editor
 {
 
 void FUtils::GetAllEntryNames(const UAnimNextRigVMAssetEditorData* InEditorData, TSet<FName>& OutNames)
@@ -363,7 +363,7 @@ bool FUtils::IsValidParameterName(const FName InName, FText& OutErrorText)
 bool FUtils::AddSchemaRigUnitAction(const TSubclassOf<URigVMUnitNode>& UnitNodeClass, UScriptStruct* Struct, const FRigVMFunction& Function, FGraphContextMenuBuilder& InContexMenuBuilder)
 {
 	FString CategoryMetadata, DisplayNameMetadata, MenuDescSuffixMetadata;
-	UE::AnimNext::Editor::FUtils::GetRigUnitStructMetadata(Struct, CategoryMetadata, DisplayNameMetadata, MenuDescSuffixMetadata);
+	UE::UAF::Editor::FUtils::GetRigUnitStructMetadata(Struct, CategoryMetadata, DisplayNameMetadata, MenuDescSuffixMetadata);
 	if (DisplayNameMetadata.IsEmpty())
 	{
 		DisplayNameMetadata = Function.GetMethodName().ToString();
diff --git a/Engine/Plugins/Experimental/AnimNext/Source/AnimNextEditor/Private/Graph/AnimNextEdGraphCustomization.h b/Engine/Plugins/Experimental/AnimNext/Source/AnimNextEditor/Private/Graph/AnimNextEdGraphCustomization.h
--- a/Engine/Plugins/Experimental/AnimNext/Source/AnimNextEditor/Private/Graph/AnimNextEdGraphCustomization.h
+++ b/Engine/Plugins/Experimental/AnimNext/Source/AnimNextEditor/Private/Graph/AnimNextEdGraphCustomization.h
@@ -15,12 +15,12 @@ namespace UE::Workspace
 	class IWorkspaceEditor;
 }
 
-namespace UE::AnimNext::Editor
+namespace UE::UAF::Editor
 {
 	class FAnimNextEditorModule;
 }
 
-namespace UE::AnimNext::Editor
+namespace UE::UAF::Editor
 {
 
 class FAnimNextEdGraphCustomization : public IDetailCustomization
diff --git a/Engine/Plugins/Experimental/AnimNext/Source/AnimNextEditor/Private/Graph/AnimNextGraphPanelPinFactory.cpp b/Engine/Plugins/Experimental/AnimNext/Source/AnimNextEditor/Private/Graph/AnimNextGraphPanelPinFactory.cpp
--- a/Engine/Plugins/Experimental/AnimNext/Source/AnimNextEditor/Private/Graph/AnimNextGraphPanelPinFactory.cpp
+++ b/Engine/Plugins/Experimental/AnimNext/Source/AnimNextEditor/Private/Graph/AnimNextGraphPanelPinFactory.cpp
@@ -21,7 +21,7 @@ TSharedPtr<SGraphPin> FAnimNextGraphPanelPinFactory::CreatePin(UEdGraphPin* InPi
 			static const FName META_CustomWidget("CustomWidget");
 			if(InPin->PinType.PinCategory == UEdGraphSchema_K2::PC_Name && ModelPin->GetMetaData(META_CustomWidget) == META_AnimNextModuleEvent)
 			{
-				return SNew(UE::AnimNext::Editor::SGraphPinModuleEvent, InPin);
+				return SNew(UE::UAF::Editor::SGraphPinModuleEvent, InPin);
 			}
 		}
 	}
@@ -31,7 +31,7 @@ TSharedPtr<SGraphPin> FAnimNextGraphPanelPinFactory::CreatePin(UEdGraphPin* InPi
 		const UScriptStruct* Struct = Cast<UScriptStruct>(InPin->PinType.PinSubCategoryObject.Get());
 		if (Struct && (Struct->IsChildOf<FAnimNextSoftVariableReference>() || Struct->IsChildOf<FAnimNextVariableReference>()))
 		{
-			return SNew(UE::AnimNext::Editor::SGraphPinVariableReference, InPin);
+			return SNew(UE::UAF::Editor::SGraphPinVariableReference, InPin);
 		}
 	}
 
diff --git a/Engine/Plugins/Experimental/AnimNext/Source/AnimNextEditor/Private/Module/AnimNextModuleAssetDefinition.cpp b/Engine/Plugins/Experimental/AnimNext/Source/AnimNextEditor/Private/Module/AnimNextModuleAssetDefinition.cpp
--- a/Engine/Plugins/Experimental/AnimNext/Source/AnimNextEditor/Private/Module/AnimNextModuleAssetDefinition.cpp
+++ b/Engine/Plugins/Experimental/AnimNext/Source/AnimNextEditor/Private/Module/AnimNextModuleAssetDefinition.cpp
@@ -15,7 +15,7 @@
 
 EAssetCommandResult UAssetDefinition_AnimNextModule::OpenAssets(const FAssetOpenArgs& OpenArgs) const
 {
-	using namespace UE::AnimNext::Editor;
+	using namespace UE::UAF::Editor;
 	using namespace UE::Workspace;
 
 	for (UAnimNextModule* Asset : OpenArgs.LoadObjects<UAnimNextModule>())
diff --git a/Engine/Plugins/Experimental/AnimNext/Source/AnimNextEditor/Private/RewindDebugger/AnimNextModuleTrack.cpp b/Engine/Plugins/Experimental/AnimNext/Source/AnimNextEditor/Private/RewindDebugger/AnimNextModuleTrack.cpp
--- a/Engine/Plugins/Experimental/AnimNext/Source/AnimNextEditor/Private/RewindDebugger/AnimNextModuleTrack.cpp
+++ b/Engine/Plugins/Experimental/AnimNext/Source/AnimNextEditor/Private/RewindDebugger/AnimNextModuleTrack.cpp
@@ -16,7 +16,7 @@
 
 #define LOCTEXT_NAMESPACE "AnimNextModuleTrack"
 
-namespace UE::AnimNextEditor
+namespace UE::UAF::Editor
 {
 
 static const FName AnimNextModulesName("AnimNextModules");
@@ -163,7 +163,7 @@ bool FAnimNextModuleTrack::UpdateInternal()
 								// Load the property descriptions
 								FMemoryReader DescriptionReader(DescriptionData->Data);
 								FObjectAsTraceIdProxyArchiveReader DescriptionReaderProxy(DescriptionReader, GameplayProvider);
-								DescriptionReaderProxy.UsingCustomVersion(AnimNext::FAnimNextTrace::CustomVersionGUID);
+								DescriptionReaderProxy.UsingCustomVersion(UE::UAF::FAnimNextTrace::CustomVersionGUID);
 								DescriptionReaderProxy << PropertyDescriptions;
 								DetailsObject->Properties.SetNum(DetailsObject->Properties.Num()+1);
 								FInstancedPropertyBag &PropertyBag = DetailsObject->Properties.Last();
diff --git a/Engine/Plugins/Experimental/AnimNext/Source/AnimNextEditor/Private/RewindDebugger/AnimNextModuleTrack.h b/Engine/Plugins/Experimental/AnimNext/Source/AnimNextEditor/Private/RewindDebugger/AnimNextModuleTrack.h
--- a/Engine/Plugins/Experimental/AnimNext/Source/AnimNextEditor/Private/RewindDebugger/AnimNextModuleTrack.h
+++ b/Engine/Plugins/Experimental/AnimNext/Source/AnimNextEditor/Private/RewindDebugger/AnimNextModuleTrack.h
@@ -28,7 +28,7 @@ public:
 };
 
 #if ANIMNEXT_TRACE_ENABLED
-namespace UE::AnimNextEditor
+namespace UE::UAF::Editor
 {
 	
 class FAnimNextModuleTrack : public RewindDebugger::FRewindDebuggerTrack
diff --git a/Engine/Plugins/Experimental/AnimNext/Source/AnimNextEditor/Private/Variables/AnimNextVariableItemMenuContext.h b/Engine/Plugins/Experimental/AnimNext/Source/AnimNextEditor/Private/Variables/AnimNextVariableItemMenuContext.h
--- a/Engine/Plugins/Experimental/AnimNext/Source/AnimNextEditor/Private/Variables/AnimNextVariableItemMenuContext.h
+++ b/Engine/Plugins/Experimental/AnimNext/Source/AnimNextEditor/Private/Variables/AnimNextVariableItemMenuContext.h
@@ -13,7 +13,7 @@ namespace UE::Workspace
 	class IWorkspaceEditor;
 }
 
-namespace UE::AnimNext::Editor
+namespace UE::UAF::Editor
 {
 	class SVariablesOutliner;
 }
@@ -32,7 +32,7 @@ public:
 	UPROPERTY()
 	TArray<TWeakObjectPtr<UAnimNextRigVMAssetEntry>> WeakEntries;
 
-	TWeakPtr<UE::AnimNext::Editor::SVariablesOutliner> WeakOutliner;
+	TWeakPtr<UE::UAF::Editor::SVariablesOutliner> WeakOutliner;
 
 	TWeakPtr<UE::Workspace::IWorkspaceEditor> WeakWorkspaceEditor;
 
diff --git a/Engine/Plugins/Experimental/AnimNext/Source/AnimNextEditor/Private/Variables/Outliner/VariablesOutlinerCategoryItem.cpp b/Engine/Plugins/Experimental/AnimNext/Source/AnimNextEditor/Private/Variables/Outliner/VariablesOutlinerCategoryItem.cpp
--- a/Engine/Plugins/Experimental/AnimNext/Source/AnimNextEditor/Private/Variables/Outliner/VariablesOutlinerCategoryItem.cpp
+++ b/Engine/Plugins/Experimental/AnimNext/Source/AnimNextEditor/Private/Variables/Outliner/VariablesOutlinerCategoryItem.cpp
@@ -15,7 +15,7 @@
 
 #define LOCTEXT_NAMESPACE "VariablesOutlinerCategoryItem"
 
-namespace UE::AnimNext::Editor
+namespace UE::UAF::Editor
 {
 class FVariableDragDropOp;
 
@@ -294,7 +294,7 @@ void FVariablesOutlinerCategoryItem::Rename(const FText& InNewName) const
 {
 	if(UAnimNextRigVMAsset* Owner = WeakOwner.Get())
 	{
-		if (UAnimNextRigVMAssetEditorData* EditorData = UE::AnimNext::UncookedOnly::FUtils::GetEditorData(Owner))
+		if (UAnimNextRigVMAssetEditorData* EditorData = UE::UAF::UncookedOnly::FUtils::GetEditorData(Owner))
 		{
 			FScopedTransaction Transaction(LOCTEXT("RenameCategory", "Rename category"));
 			EditorData->RenameCategory(CategoryName, InNewName.ToString());						
@@ -306,7 +306,7 @@ bool FVariablesOutlinerCategoryItem::ValidateName(const FText& InNewName, FText&
 {
 	if(UAnimNextRigVMAsset* Owner = WeakOwner.Get())
 	{
-		if (UAnimNextRigVMAssetEditorData* EditorData = UE::AnimNext::UncookedOnly::FUtils::GetEditorData(Owner))
+		if (UAnimNextRigVMAssetEditorData* EditorData = UE::UAF::UncookedOnly::FUtils::GetEditorData(Owner))
 		{
 			if (EditorData->VariableAndFunctionCategories.Contains(InNewName.ToString()))
 			{
diff --git a/Engine/Plugins/Experimental/AnimNext/Source/AnimNextEditor/Private/Variables/Outliner/VariablesOutlinerMode.cpp b/Engine/Plugins/Experimental/AnimNext/Source/AnimNextEditor/Private/Variables/Outliner/VariablesOutlinerMode.cpp
--- a/Engine/Plugins/Experimental/AnimNext/Source/AnimNextEditor/Private/Variables/Outliner/VariablesOutlinerMode.cpp
+++ b/Engine/Plugins/Experimental/AnimNext/Source/AnimNextEditor/Private/Variables/Outliner/VariablesOutlinerMode.cpp
@@ -39,7 +39,7 @@
 
 #define LOCTEXT_NAMESPACE "VariablesOutlinerMode"
 
-namespace UE::AnimNext::Editor
+namespace UE::UAF::Editor
 {
 
 FVariablesOutlinerMode::FVariablesOutlinerMode(SVariablesOutliner* InVariablesOutliner, const TSharedRef<UE::Workspace::IWorkspaceEditor>& InWorkspaceEditor)
@@ -116,7 +116,7 @@ TSharedPtr<SWidget> FVariablesOutlinerMode::CreateContextMenu()
 							}),
 							FUIAction(FExecuteAction(), FCanExecuteAction(), FGetActionCheckState(), FIsActionButtonVisible::CreateLambda([MenuContext]() -> bool
 							{
-								if (TSharedPtr<UE::AnimNext::Editor::SVariablesOutliner> SharedOutliner = MenuContext->WeakOutliner.Pin())
+								if (TSharedPtr<UE::UAF::Editor::SVariablesOutliner> SharedOutliner = MenuContext->WeakOutliner.Pin())
 								{
 									if (const FVariablesOutlinerMode* Mode = static_cast<const FVariablesOutlinerMode*>(SharedOutliner->GetMode()))
 									{
@@ -599,7 +599,7 @@ void FVariablesOutlinerMode::FindReferences(ESearchScope InSearchScope) const
 		FSceneOutlinerTreeItemPtr ItemToFind = Selection.SelectedItems[0].Pin();
 		if (const FVariablesOutlinerEntryItem* EntryItem = ItemToFind->CastTo<FVariablesOutlinerEntryItem>())
 		{
-			if (TSharedPtr<SDockTab> FindAndReplaceTab = WeakWorkspaceEditor.Pin()->GetTabManager()->TryInvokeTab(FTabId(UE::AnimNext::Editor::FindAndReplaceTabName)))
+			if (TSharedPtr<SDockTab> FindAndReplaceTab = WeakWorkspaceEditor.Pin()->GetTabManager()->TryInvokeTab(FTabId(UE::UAF::Editor::FindAndReplaceTabName)))
 			{
 				TSharedRef<IAnimAssetFindReplace> FindAndReplace = StaticCastSharedRef<IAnimAssetFindReplace>(FindAndReplaceTab->GetContent());
 				FindAndReplace->SetCurrentProcessor(UAnimNextAssetFindReplaceVariables::StaticClass());
@@ -656,7 +656,7 @@ void FVariablesOutlinerMode::PopulateNewVariableToolMenuEntries(UToolMenu* InMen
 	}
 	else
 	{
-		TSharedPtr<UE::AnimNext::Editor::SVariablesOutliner> SharedOutliner = MenuContext->WeakOutliner.Pin();
+		TSharedPtr<UE::UAF::Editor::SVariablesOutliner> SharedOutliner = MenuContext->WeakOutliner.Pin();
 		if (SharedOutliner.IsValid())
 		{
 			TArray<FSceneOutlinerTreeItemPtr> SelectedItems = SharedOutliner->GetSelectedItems();
diff --git a/Engine/Plugins/Experimental/AnimNext/Source/AnimNextEditor/Private/Variables/SGraphPinVariableReference.cpp b/Engine/Plugins/Experimental/AnimNext/Source/AnimNextEditor/Private/Variables/SGraphPinVariableReference.cpp
--- a/Engine/Plugins/Experimental/AnimNext/Source/AnimNextEditor/Private/Variables/SGraphPinVariableReference.cpp
+++ b/Engine/Plugins/Experimental/AnimNext/Source/AnimNextEditor/Private/Variables/SGraphPinVariableReference.cpp
@@ -16,10 +16,11 @@
 #include "RigVMModel/RigVMController.h"
 #include "Variables/AnimNextVariableReference.h"
 #include "Variables/AnimNextSoftVariableReference.h"
+#include "Widgets/Layout/SBox.h"
 
 #define LOCTEXT_NAMESPACE "SGraphPinVariableReference"
 
-namespace UE::AnimNext::Editor
+namespace UE::UAF::Editor
 {
 
 void SGraphPinVariableReference::Construct(const FArguments& InArgs, UEdGraphPin* InPin)
diff --git a/Engine/Plugins/Experimental/AnimNext/Source/AnimNextEditor/Private/Workspace/AnimNextWorkspaceEditorMode.cpp b/Engine/Plugins/Experimental/AnimNext/Source/AnimNextEditor/Private/Workspace/AnimNextWorkspaceEditorMode.cpp
--- a/Engine/Plugins/Experimental/AnimNext/Source/AnimNextEditor/Private/Workspace/AnimNextWorkspaceEditorMode.cpp
+++ b/Engine/Plugins/Experimental/AnimNext/Source/AnimNextEditor/Private/Workspace/AnimNextWorkspaceEditorMode.cpp
@@ -47,7 +47,7 @@ void UAnimNextWorkspaceEditorMode::Enter()
 	for(UObject* Asset : Assets)
 	{
 		UAnimNextRigVMAsset* AnimNextRigVMAsset = static_cast<UAnimNextRigVMAsset*>(Asset);
-		UAnimNextRigVMAssetEditorData* EditorData = UE::AnimNext::UncookedOnly::FUtils::GetEditorData<UAnimNextRigVMAssetEditorData>(AnimNextRigVMAsset);
+		UAnimNextRigVMAssetEditorData* EditorData = UE::UAF::UncookedOnly::FUtils::GetEditorData<UAnimNextRigVMAssetEditorData>(AnimNextRigVMAsset);
 
 		EditorData->RigVMCompiledEvent.RemoveAll(this);
 		EditorData->RigVMCompiledEvent.AddUObject(this, &UAnimNextWorkspaceEditorMode::HandleRigVMCompiledEvent);
@@ -64,9 +64,9 @@ void UAnimNextWorkspaceEditorMode::Exit()
 {
 	Super::Exit();
 
-	for (TPair<FObjectKey, TSharedRef<UE::AnimNext::Editor::IAssetCompilationHandler>>& AssetCompiler : AssetCompilers)
+	for (TPair<FObjectKey, TSharedRef<UE::UAF::Editor::IAssetCompilationHandler>>& AssetCompiler : AssetCompilers)
 	{
-		TSharedRef<UE::AnimNext::Editor::IAssetCompilationHandler> AssetCompilerRef = AssetCompiler.Value;
+		TSharedRef<UE::UAF::Editor::IAssetCompilationHandler> AssetCompilerRef = AssetCompiler.Value;
 		AssetCompilerRef->OnCompileStatusChanged().Unbind();
 	}
 
@@ -80,7 +80,7 @@ void UAnimNextWorkspaceEditorMode::Exit()
 			continue;;
 		}
 		
-		UAnimNextRigVMAssetEditorData* EditorData = UE::AnimNext::UncookedOnly::FUtils::GetEditorData<UAnimNextRigVMAssetEditorData>(AnimNextRigVMAsset);
+		UAnimNextRigVMAssetEditorData* EditorData = UE::UAF::UncookedOnly::FUtils::GetEditorData<UAnimNextRigVMAssetEditorData>(AnimNextRigVMAsset);
 
 		EditorData->RigVMCompiledEvent.RemoveAll(this);
 		EditorData->RigVMGraphModifiedEvent.RemoveAll(this);
@@ -101,7 +101,7 @@ void UAnimNextWorkspaceEditorMode::BindCommands()
 
 	TSharedRef<FUICommandList> ToolkitCommands = WorkspaceEditor->GetToolkitCommands();
 
-	const UE::AnimNext::FRigVMCommands& RigVMCommands = UE::AnimNext::FRigVMCommands::Get();
+	const UE::UAF::FRigVMCommands& RigVMCommands = UE::UAF::FRigVMCommands::Get();
 	ToolkitCommands->MapAction(RigVMCommands.Compile,
 		FExecuteAction::CreateUObject(this, &UAnimNextWorkspaceEditorMode::HandleCompile));
 
@@ -169,11 +169,11 @@ void UAnimNextWorkspaceEditorMode::HandleCompile()
 		for(int32 AssetIndex = 0; AssetIndex < Assets.Num(); ++AssetIndex)
 		{
 			UObject* AssetToCheck = Assets[AssetIndex];
-			if(TSharedPtr<UE::AnimNext::Editor::IAssetCompilationHandler> FoundCompiler = GetAssetCompiler(AssetToCheck))
+			if(TSharedPtr<UE::UAF::Editor::IAssetCompilationHandler> FoundCompiler = GetAssetCompiler(AssetToCheck))
 			{
-				UE::AnimNext::Editor::ECompileStatus AssetCompileStatus = FoundCompiler->GetCompileStatus(WorkspaceEditor.ToSharedRef(), AssetToCheck);
-				if( AssetCompileStatus != UE::AnimNext::Editor::ECompileStatus::Dirty &&
-					AssetCompileStatus != UE::AnimNext::Editor::ECompileStatus::Error)
+				UE::UAF::Editor::ECompileStatus AssetCompileStatus = FoundCompiler->GetCompileStatus(WorkspaceEditor.ToSharedRef(), AssetToCheck);
+				if( AssetCompileStatus != UE::UAF::Editor::ECompileStatus::Dirty &&
+					AssetCompileStatus != UE::UAF::Editor::ECompileStatus::Error)
 				{
 					Assets.RemoveAtSwap(AssetIndex);
 					AssetIndex--;
@@ -188,12 +188,12 @@ void UAnimNextWorkspaceEditorMode::HandleCompile()
 	}
 
 	// Start a batch compile scope
-	UE::AnimNext::UncookedOnly::FScopedCompileJob CompileJob(JobName, Assets);
+	UE::UAF::UncookedOnly::FScopedCompileJob CompileJob(JobName, Assets);
 
-	UE::AnimNext::Editor::FAnimNextEditorModule& AnimNextEditorModule = FModuleManager::LoadModuleChecked<UE::AnimNext::Editor::FAnimNextEditorModule>("AnimNextEditor");
+	UE::UAF::Editor::FAnimNextEditorModule& AnimNextEditorModule = FModuleManager::LoadModuleChecked<UE::UAF::Editor::FAnimNextEditorModule>("AnimNextEditor");
 	for(UObject* AssetToCompile : Assets)
 	{
-		if(TSharedPtr<UE::AnimNext::Editor::IAssetCompilationHandler> FoundCompiler = GetAssetCompiler(AssetToCompile))
+		if(TSharedPtr<UE::UAF::Editor::IAssetCompilationHandler> FoundCompiler = GetAssetCompiler(AssetToCompile))
 		{
 			FoundCompiler->Compile(WorkspaceEditor.ToSharedRef(), AssetToCompile);
 		}
@@ -215,10 +215,10 @@ void UAnimNextWorkspaceEditorMode::PropagateAutoCompile(TSharedRef<UE::Workspace
 {
 	TArray<UObject*> Assets;
 	InWorkspaceEditor->GetOpenedAssets(Assets);
-	UE::AnimNext::Editor::FAnimNextEditorModule& AnimNextEditorModule = FModuleManager::LoadModuleChecked<UE::AnimNext::Editor::FAnimNextEditorModule>("AnimNextEditor");
+	UE::UAF::Editor::FAnimNextEditorModule& AnimNextEditorModule = FModuleManager::LoadModuleChecked<UE::UAF::Editor::FAnimNextEditorModule>("AnimNextEditor");
 	for(UObject* Asset : Assets)
 	{
-		if(TSharedPtr<UE::AnimNext::Editor::IAssetCompilationHandler> FoundCompiler = GetAssetCompiler(Asset))
+		if(TSharedPtr<UE::UAF::Editor::IAssetCompilationHandler> FoundCompiler = GetAssetCompiler(Asset))
 		{
 			FoundCompiler->SetAutoCompile(InWorkspaceEditor, Asset, bInAutoCompile);
 		}
@@ -271,7 +271,7 @@ void UAnimNextWorkspaceEditorMode::HandleFocusedDocumentChanged(const UE::Worksp
 	}
 
 	// Ensure auto-compilation is propagated for newly opened asset 
-	UAnimNextRigVMAssetEditorData* EditorData = UE::AnimNext::UncookedOnly::FUtils::GetEditorData<UAnimNextRigVMAssetEditorData>(AnimNextRigVMAsset);
+	UAnimNextRigVMAssetEditorData* EditorData = UE::UAF::UncookedOnly::FUtils::GetEditorData<UAnimNextRigVMAssetEditorData>(AnimNextRigVMAsset);
 	EditorData->SetAutoVMRecompile(State.bAutoCompile);
 
 	// Subscribe to asset compilation/modification
@@ -302,20 +302,20 @@ void UAnimNextWorkspaceEditorMode::UpdateCompileStatus()
 
 	for(UObject* Asset : Assets)
 	{
-		if(TSharedPtr<UE::AnimNext::Editor::IAssetCompilationHandler> FoundCompiler = GetAssetCompiler(Asset))
+		if(TSharedPtr<UE::UAF::Editor::IAssetCompilationHandler> FoundCompiler = GetAssetCompiler(Asset))
 		{
-			UE::AnimNext::Editor::ECompileStatus AssetCompileStatus = FoundCompiler->GetCompileStatus(WorkspaceEditor.ToSharedRef(), Asset);
-			if(AssetCompileStatus == UE::AnimNext::Editor::ECompileStatus::Dirty)
+			UE::UAF::Editor::ECompileStatus AssetCompileStatus = FoundCompiler->GetCompileStatus(WorkspaceEditor.ToSharedRef(), Asset);
+			if(AssetCompileStatus == UE::UAF::Editor::ECompileStatus::Dirty)
 			{
 				NumAssetsDirty++;
 			}
 
-			if(AssetCompileStatus == UE::AnimNext::Editor::ECompileStatus::Error)
+			if(AssetCompileStatus == UE::UAF::Editor::ECompileStatus::Error)
 			{
 				NumAssetsWithErrors++;
 			}
 
-			if(AssetCompileStatus == UE::AnimNext::Editor::ECompileStatus::Warning)
+			if(AssetCompileStatus == UE::UAF::Editor::ECompileStatus::Warning)
 			{
 				NumAssetsWithWarnings++;
 			}
@@ -324,19 +324,19 @@ void UAnimNextWorkspaceEditorMode::UpdateCompileStatus()
 
 	if(NumAssetsWithErrors > 0)
 	{
-		CompileStatus = UE::AnimNext::Editor::ECompileStatus::Error;
+		CompileStatus = UE::UAF::Editor::ECompileStatus::Error;
 	}
 	else if(NumAssetsWithWarnings > 0)
 	{
-		CompileStatus = UE::AnimNext::Editor::ECompileStatus::Warning;
+		CompileStatus = UE::UAF::Editor::ECompileStatus::Warning;
 	}
 	else if(NumAssetsDirty > 0)
 	{
-		CompileStatus = UE::AnimNext::Editor::ECompileStatus::Dirty;
+		CompileStatus = UE::UAF::Editor::ECompileStatus::Dirty;
 	}
 	else
 	{
-		CompileStatus = UE::AnimNext::Editor::ECompileStatus::UpToDate;
+		CompileStatus = UE::UAF::Editor::ECompileStatus::UpToDate;
 	}
 }
 
@@ -355,7 +355,7 @@ void UAnimNextWorkspaceEditorMode::HandleRigVMModifiedEvent(ERigVMGraphNotifType
 	UAnimNextRigVMAssetEditorData* EditorData = InGraph->GetTypedOuter<UAnimNextRigVMAssetEditorData>();
 	if(EditorData->IsDirtyForRecompilation())
 	{
-		CompileStatus = UE::AnimNext::Editor::ECompileStatus::Dirty;
+		CompileStatus = UE::UAF::Editor::ECompileStatus::Dirty;
 	}
 }
 
@@ -372,17 +372,17 @@ TSharedPtr<UE::Workspace::IWorkspaceEditor> UAnimNextWorkspaceEditorMode::GetWor
 	return nullptr;
 }
 
-TSharedPtr<UE::AnimNext::Editor::IAssetCompilationHandler> UAnimNextWorkspaceEditorMode::GetAssetCompiler(UObject* InAsset)
+TSharedPtr<UE::UAF::Editor::IAssetCompilationHandler> UAnimNextWorkspaceEditorMode::GetAssetCompiler(UObject* InAsset)
 {
-	if(TSharedRef<UE::AnimNext::Editor::IAssetCompilationHandler>* FoundCompiler = AssetCompilers.Find(InAsset))
+	if(TSharedRef<UE::UAF::Editor::IAssetCompilationHandler>* FoundCompiler = AssetCompilers.Find(InAsset))
 	{
 		return *FoundCompiler;
 	}
 
-	UE::AnimNext::Editor::FAnimNextEditorModule& AnimNextEditorModule = FModuleManager::LoadModuleChecked<UE::AnimNext::Editor::FAnimNextEditorModule>("AnimNextEditor");
-	if(const UE::AnimNext::Editor::FAssetCompilationHandlerFactoryDelegate* FoundFactory = AnimNextEditorModule.FindAssetCompilationHandlerFactory(InAsset->GetClass()))
+	UE::UAF::Editor::FAnimNextEditorModule& AnimNextEditorModule = FModuleManager::LoadModuleChecked<UE::UAF::Editor::FAnimNextEditorModule>("AnimNextEditor");
+	if(const UE::UAF::Editor::FAssetCompilationHandlerFactoryDelegate* FoundFactory = AnimNextEditorModule.FindAssetCompilationHandlerFactory(InAsset->GetClass()))
 	{
-		TSharedRef<UE::AnimNext::Editor::IAssetCompilationHandler> NewCompiler = AssetCompilers.Add(InAsset, FoundFactory->Execute(InAsset));
+		TSharedRef<UE::UAF::Editor::IAssetCompilationHandler> NewCompiler = AssetCompilers.Add(InAsset, FoundFactory->Execute(InAsset));
 		NewCompiler->OnCompileStatusChanged().BindUObject(this, &UAnimNextWorkspaceEditorMode::UpdateCompileStatus);
 		return NewCompiler;
 	}
