UAF: Namespace renamed

diff --git a/Engine/Plugins/Experimental/AnimNext/Source/AnimNextEditor/Private/Workspace/AnimNextWorkspaceEditorMode.h b/Engine/Plugins/Experimental/AnimNext/Source/AnimNextEditor/Private/Workspace/AnimNextWorkspaceEditorMode.h
--- a/Engine/Plugins/Experimental/AnimNext/Source/AnimNextEditor/Private/Workspace/AnimNextWorkspaceEditorMode.h
+++ b/Engine/Plugins/Experimental/AnimNext/Source/AnimNextEditor/Private/Workspace/AnimNextWorkspaceEditorMode.h
@@ -30,7 +30,7 @@ public:
 	UAnimNextWorkspaceEditorMode();
 
 	// Get the current compilation status
-	UE::AnimNext::Editor::ECompileStatus GetLatestCompileStatus() const { return CompileStatus; }
+	UE::UAF::Editor::ECompileStatus GetLatestCompileStatus() const { return CompileStatus; }
 
 	// Get the currently stored workspace state
 	const FAnimNextWorkspaceState& GetState() const { return State; } 
@@ -69,7 +69,7 @@ private:
 	TSharedPtr<UE::Workspace::IWorkspaceEditor> GetWorkspaceEditor() const;
 
 	// Lazily construct an asset compiler for the supplied asset
-	TSharedPtr<UE::AnimNext::Editor::IAssetCompilationHandler> GetAssetCompiler(UObject* InAsset);
+	TSharedPtr<UE::UAF::Editor::IAssetCompilationHandler> GetAssetCompiler(UObject* InAsset);
 
 	friend UAnimNextWorkspaceSchema;
 
@@ -77,13 +77,13 @@ private:
 	FAnimNextWorkspaceState State;
 
 	// Asset compilers for all current assets 
-	TMap<FObjectKey, TSharedRef<UE::AnimNext::Editor::IAssetCompilationHandler>> AssetCompilers;
+	TMap<FObjectKey, TSharedRef<UE::UAF::Editor::IAssetCompilationHandler>> AssetCompilers;
 
 	// All assets that we are currently tracking for compilation status
 	TSet<TObjectKey<UObject>> WeakAssets;
 
 	// Current compilation status
-	UE::AnimNext::Editor::ECompileStatus CompileStatus = UE::AnimNext::Editor::ECompileStatus::Unknown;
+	UE::UAF::Editor::ECompileStatus CompileStatus = UE::UAF::Editor::ECompileStatus::Unknown;
 
 	// Flag to update the latest compile status
 	bool bUpdateCompileStatus = false;
diff --git a/Engine/Plugins/Experimental/AnimNext/Source/AnimNextEditor/Private/Workspace/AnimNextWorkspaceSchema.cpp b/Engine/Plugins/Experimental/AnimNext/Source/AnimNextEditor/Private/Workspace/AnimNextWorkspaceSchema.cpp
--- a/Engine/Plugins/Experimental/AnimNext/Source/AnimNextEditor/Private/Workspace/AnimNextWorkspaceSchema.cpp
+++ b/Engine/Plugins/Experimental/AnimNext/Source/AnimNextEditor/Private/Workspace/AnimNextWorkspaceSchema.cpp
@@ -17,7 +17,7 @@ FText UAnimNextWorkspaceSchema::GetDisplayName() const
 
 TConstArrayView<FTopLevelAssetPath> UAnimNextWorkspaceSchema::GetSupportedAssetClassPaths() const
 {
-	const UE::AnimNext::Editor::FAnimNextEditorModule& Module = FModuleManager::Get().LoadModuleChecked<UE::AnimNext::Editor::FAnimNextEditorModule>("AnimNextEditor");
+	const UE::UAF::Editor::FAnimNextEditorModule& Module = FModuleManager::Get().LoadModuleChecked<UE::UAF::Editor::FAnimNextEditorModule>("AnimNextEditor");
 
 	return Module.SupportedAssetClasses;
 }
diff --git a/Engine/Plugins/Experimental/AnimNext/Source/AnimNextEditor/Public/FindInAnimNextRigVMAsset.h b/Engine/Plugins/Experimental/AnimNext/Source/AnimNextEditor/Public/FindInAnimNextRigVMAsset.h
--- a/Engine/Plugins/Experimental/AnimNext/Source/AnimNextEditor/Public/FindInAnimNextRigVMAsset.h
+++ b/Engine/Plugins/Experimental/AnimNext/Source/AnimNextEditor/Public/FindInAnimNextRigVMAsset.h
@@ -8,7 +8,7 @@
 class FAssetEditorToolkit;
 class SSearchBox;
 
-namespace UE::AnimNext::Editor
+namespace UE::UAF::Editor
 {
 
 /** Item that matched the search results */
@@ -37,4 +37,4 @@ protected:
 	//~ End SFindInGraph Interface
 };
 
-} // namespace UE::AnimNext::Editor
\ No newline at end of file
+} // namespace UE::UAF::Editor
\ No newline at end of file
diff --git a/Engine/Plugins/Experimental/AnimNext/Source/AnimNextTestSuite/Private/AnimNextEditorTest.cpp b/Engine/Plugins/Experimental/AnimNext/Source/AnimNextTestSuite/Private/AnimNextEditorTest.cpp
--- a/Engine/Plugins/Experimental/AnimNext/Source/AnimNextTestSuite/Private/AnimNextEditorTest.cpp
+++ b/Engine/Plugins/Experimental/AnimNext/Source/AnimNextTestSuite/Private/AnimNextEditorTest.cpp
@@ -21,14 +21,14 @@
 
 #if WITH_DEV_AUTOMATION_TESTS && WITH_EDITOR
 
-namespace UE::AnimNext::Tests
+namespace UE::UAF::Tests
 {
 
 IMPLEMENT_SIMPLE_AUTOMATION_TEST(FEditor_Variables, "Animation.AnimNext.Editor.Variables", EAutomationTestFlags::EditorContext | EAutomationTestFlags::EngineFilter)
 
 bool FEditor_Variables::RunTest(const FString& InParameters)
 {
-	using namespace UE::AnimNext;
+	using namespace UE::UAF;
 
 	struct FFactoryAndClass
 	{
@@ -127,7 +127,7 @@ IMPLEMENT_SIMPLE_AUTOMATION_TEST(FEditor_Graphs, "Animation.AnimNext.Editor.Grap
 
 bool FEditor_Graphs::RunTest(const FString& InParameters)
 {
-	using namespace UE::AnimNext;
+	using namespace UE::UAF;
 
 	struct FTestSettings
 	{
@@ -199,7 +199,7 @@ IMPLEMENT_SIMPLE_AUTOMATION_TEST(FEditor_Variables_Python, "Animation.AnimNext.E
 
 bool FEditor_Variables_Python::RunTest(const FString& InParameters)
 {
-	using namespace UE::AnimNext;
+	using namespace UE::UAF;
 
 	const TCHAR* Script = TEXT(
 		"asset_tools = unreal.AssetToolsHelpers.get_asset_tools()\n"
diff --git a/Engine/Plugins/Experimental/AnimNext/Source/AnimNextTestSuite/Private/AnimNextTest.cpp b/Engine/Plugins/Experimental/AnimNext/Source/AnimNextTestSuite/Private/AnimNextTest.cpp
--- a/Engine/Plugins/Experimental/AnimNext/Source/AnimNextTestSuite/Private/AnimNextTest.cpp
+++ b/Engine/Plugins/Experimental/AnimNext/Source/AnimNextTestSuite/Private/AnimNextTest.cpp
@@ -14,7 +14,7 @@
 
 #if WITH_DEV_AUTOMATION_TESTS
 
-namespace UE::AnimNext::Tests
+namespace UE::UAF::Tests
 {
 
 namespace Private
@@ -57,7 +57,7 @@ IMPLEMENT_SIMPLE_AUTOMATION_TEST(FDataRegistryTransformsTest, "Animation.AnimNex
 
 bool FDataRegistryTransformsTest::RunTest(const FString& InParameters)
 {
-	using namespace UE::AnimNext;
+	using namespace UE::UAF;
 
 	FDataRegistry* AnimDataReg = FDataRegistry::Get();
 
diff --git a/Engine/Plugins/Experimental/AnimNext/Source/AnimNextTestSuite/Private/AnimNextVariablesTest.cpp b/Engine/Plugins/Experimental/AnimNext/Source/AnimNextTestSuite/Private/AnimNextVariablesTest.cpp
--- a/Engine/Plugins/Experimental/AnimNext/Source/AnimNextTestSuite/Private/AnimNextVariablesTest.cpp
+++ b/Engine/Plugins/Experimental/AnimNext/Source/AnimNextTestSuite/Private/AnimNextVariablesTest.cpp
@@ -26,7 +26,7 @@
 
 #if WITH_DEV_AUTOMATION_TESTS
 
-namespace UE::AnimNext::Tests
+namespace UE::UAF::Tests
 {
 
 IMPLEMENT_SIMPLE_AUTOMATION_TEST(FVariableTypesTest, "Animation.AnimNext.VariableTypes", EAutomationTestFlags::EditorContext | EAutomationTestFlags::EngineFilter)
diff --git a/Engine/Plugins/Experimental/AnimNext/Source/AnimNextTestSuite/Private/Module.cpp b/Engine/Plugins/Experimental/AnimNext/Source/AnimNextTestSuite/Private/Module.cpp
--- a/Engine/Plugins/Experimental/AnimNext/Source/AnimNextTestSuite/Private/Module.cpp
+++ b/Engine/Plugins/Experimental/AnimNext/Source/AnimNextTestSuite/Private/Module.cpp
@@ -4,7 +4,7 @@
 #include "Modules/ModuleInterface.h"
 #include "Modules/ModuleManager.h"
 
-namespace UE::AnimNext::Tests
+namespace UE::UAF::Tests
 {
 
 class FModule : public IModuleInterface
@@ -14,4 +14,4 @@ public:
 
 }
 
-IMPLEMENT_MODULE(UE::AnimNext::Tests::FModule, AnimNextTestSuite)
+IMPLEMENT_MODULE(UE::UAF::Tests::FModule, AnimNextTestSuite)
diff --git a/Engine/Plugins/Experimental/AnimNext/Source/AnimNextTestSuite/Public/AnimNextTest.h b/Engine/Plugins/Experimental/AnimNext/Source/AnimNextTestSuite/Public/AnimNextTest.h
--- a/Engine/Plugins/Experimental/AnimNext/Source/AnimNextTestSuite/Public/AnimNextTest.h
+++ b/Engine/Plugins/Experimental/AnimNext/Source/AnimNextTestSuite/Public/AnimNextTest.h
@@ -16,7 +16,7 @@ struct FAnimNextTestData
 	float B = 0.f;
 };
 
-namespace UE::AnimNext::Tests
+namespace UE::UAF::Tests
 {
 
 struct FUtils final
diff --git a/Engine/Plugins/Experimental/AnimNext/Source/AnimNextUncookedOnly/Internal/AnimNextEdGraph.h b/Engine/Plugins/Experimental/AnimNext/Source/AnimNextUncookedOnly/Internal/AnimNextEdGraph.h
--- a/Engine/Plugins/Experimental/AnimNext/Source/AnimNextUncookedOnly/Internal/AnimNextEdGraph.h
+++ b/Engine/Plugins/Experimental/AnimNext/Source/AnimNextUncookedOnly/Internal/AnimNextEdGraph.h
@@ -8,7 +8,7 @@
 
 class UAnimNextRigVMAssetEditorData;
 
-namespace UE::AnimNext::UncookedOnly
+namespace UE::UAF::UncookedOnly
 {
 	struct FUtils;
 }
diff --git a/Engine/Plugins/Experimental/AnimNext/Source/AnimNextUncookedOnly/Internal/AnimNextRigVMAssetEditorData.h b/Engine/Plugins/Experimental/AnimNext/Source/AnimNextUncookedOnly/Internal/AnimNextRigVMAssetEditorData.h
--- a/Engine/Plugins/Experimental/AnimNext/Source/AnimNextUncookedOnly/Internal/AnimNextRigVMAssetEditorData.h
+++ b/Engine/Plugins/Experimental/AnimNext/Source/AnimNextUncookedOnly/Internal/AnimNextRigVMAssetEditorData.h
@@ -28,7 +28,7 @@ struct FAnimNextGetVariableCompileContext;
 struct FAnimNextGetGraphCompileContext;
 struct FAnimNextProcessGraphCompileContext;
 
-namespace UE::AnimNext::UncookedOnly
+namespace UE::UAF::UncookedOnly
 {
 	struct FPublicVariablesImpl;
 	struct FUtils;
@@ -36,7 +36,7 @@ namespace UE::AnimNext::UncookedOnly
 	class FScopedCompileJob;
 }
 
-namespace UE::AnimNext::Editor
+namespace UE::UAF::Editor
 {
 	struct FUtils;
 	class SRigVMAssetView;
@@ -64,7 +64,7 @@ namespace UE::AnimNext::Editor
 	struct FVariablesOutlinerCategoryItem;
 }
 
-namespace UE::AnimNext::Tests
+namespace UE::UAF::Tests
 {
 	class FEditor_Graphs;
 	class FEditor_AnimGraph_Graphs;
@@ -92,7 +92,7 @@ enum class EAnimNextEditorDataNotifType : uint8
 	VariablesReordered,  // Asset entries have been reordered (Subject == UAnimNextRigVMAssetEditorData)
 };
 
-namespace UE::AnimNext::UncookedOnly
+namespace UE::UAF::UncookedOnly
 {
 	// A delegate for subscribing / reacting to editor data modifications.
 	DECLARE_MULTICAST_DELEGATE_ThreeParams(FOnEditorDataModified, UAnimNextRigVMAssetEditorData* /* InEditorData */, EAnimNextEditorDataNotifType /* InType */, UObject* /* InSubject */);
@@ -205,45 +205,45 @@ public:
 	static UE_API void ReportError(const TCHAR* InMessage);
 
 protected:
-	friend class UE::AnimNext::Editor::SRigVMAssetView;
-	friend class UE::AnimNext::Editor::SRigVMAssetViewRow;
-	friend struct UE::AnimNext::Editor::FUtils;
-	friend struct UE::AnimNext::UncookedOnly::FUtils;
-	friend class UE::AnimNext::Editor::FAnimNextEditorModule;
-	friend class UE::AnimNext::Editor::FWorkspaceEditor;
+	friend class UE::UAF::Editor::SRigVMAssetView;
+	friend class UE::UAF::Editor::SRigVMAssetViewRow;
+	friend struct UE::UAF::Editor::FUtils;
+	friend struct UE::UAF::UncookedOnly::FUtils;
+	friend class UE::UAF::Editor::FAnimNextEditorModule;
+	friend class UE::UAF::Editor::FWorkspaceEditor;
 	friend class UAnimNextRigVMAssetEntry;
 	friend class UAnimNextRigVMAssetLibrary;
 	friend class UAnimNextEdGraph;
-	friend class UE::AnimNext::Tests::FEditor_Graphs;
-	friend class UE::AnimNext::Tests::FEditor_Variables;
-	friend class UE::AnimNext::Tests::FEditor_AnimGraph_Graphs;
-	friend class UE::AnimNext::Tests::FEditor_AnimGraph_Variables;
-	friend class UE::AnimNext::Editor::FFindInAnimNextRigVMAssetResult;
-	friend class UE::AnimNext::Editor::SFindInAnimNextRigVMAsset;
-	friend class UE::AnimNext::Tests::FVariables;
-	friend class UE::AnimNext::Tests::FVariables_UOLBindings;
-	friend class UE::AnimNext::Editor::FVariableCustomization;
-	friend class UE::AnimNext::Editor::FAnimNextAssetItemDetails;
-	friend class UE::AnimNext::Editor::FAnimNextGraphItemDetails;
-	friend class UE::AnimNext::Editor::FAnimNextFunctionItemDetails;
-	friend struct UE::AnimNext::Editor::FVariablesOutlinerEntryItem;
-	friend struct UE::AnimNext::Editor::FVariablesOutlinerCategoryItem;
-	friend class UE::AnimNext::Editor::FVariablesOutlinerMode;
-	friend class UE::AnimNext::Editor::FVariablesOutlinerHierarchy;
-	friend class UE::AnimNext::Editor::SVariablesOutlinerValue;
-	friend class UE::AnimNext::Editor::SVariablesOutliner;
+	friend class UE::UAF::Tests::FEditor_Graphs;
+	friend class UE::UAF::Tests::FEditor_Variables;
+	friend class UE::UAF::Tests::FEditor_AnimGraph_Graphs;
+	friend class UE::UAF::Tests::FEditor_AnimGraph_Variables;
+	friend class UE::UAF::Editor::FFindInAnimNextRigVMAssetResult;
+	friend class UE::UAF::Editor::SFindInAnimNextRigVMAsset;
+	friend class UE::UAF::Tests::FVariables;
+	friend class UE::UAF::Tests::FVariables_UOLBindings;
+	friend class UE::UAF::Editor::FVariableCustomization;
+	friend class UE::UAF::Editor::FAnimNextAssetItemDetails;
+	friend class UE::UAF::Editor::FAnimNextGraphItemDetails;
+	friend class UE::UAF::Editor::FAnimNextFunctionItemDetails;
+	friend struct UE::UAF::Editor::FVariablesOutlinerEntryItem;
+	friend struct UE::UAF::Editor::FVariablesOutlinerCategoryItem;
+	friend class UE::UAF::Editor::FVariablesOutlinerMode;
+	friend class UE::UAF::Editor::FVariablesOutlinerHierarchy;
+	friend class UE::UAF::Editor::SVariablesOutlinerValue;
+	friend class UE::UAF::Editor::SVariablesOutliner;
 	friend class UAnimNextModuleWorkspaceAssetUserData;
-	friend class UE::AnimNext::Editor::SAddVariablesDialog;
-	friend class UE::AnimNext::Editor::FVariableProxyCustomization;
+	friend class UE::UAF::Editor::SAddVariablesDialog;
+	friend class UE::UAF::Editor::FVariableProxyCustomization;
 	friend class UAnimNextSharedVariablesEntry;
-	friend struct UE::AnimNext::UncookedOnly::FPublicVariablesImpl;
-	friend UE::AnimNext::Editor::FAnimNextAnimGraphEditorModule;
-	friend class UE::AnimNext::Editor::FCallFunctionSharedDataDetails;
+	friend struct UE::UAF::UncookedOnly::FPublicVariablesImpl;
+	friend UE::UAF::Editor::FAnimNextAnimGraphEditorModule;
+	friend class UE::UAF::Editor::FCallFunctionSharedDataDetails;
 	friend UAnimNextWorkspaceEditorMode;
-	friend UE::AnimNext::UncookedOnly::FScopedCompileJob;
-	friend UE::AnimNext::Editor::FAssetCompilationHandler;
+	friend UE::UAF::UncookedOnly::FScopedCompileJob;
+	friend UE::UAF::Editor::FAssetCompilationHandler;
 	friend class UAnimNextStateTreeTreeEditorData;
-	friend class UE::AnimNext::Editor::FVariableBindingPropertyCustomization;
+	friend class UE::UAF::Editor::FVariableBindingPropertyCustomization;
 	friend class UAnimNextAssetFindReplaceVariables;
 	friend struct FAnimNextSchemaAction_Variable;
 	friend struct FAnimNextModuleTests;
@@ -678,10 +678,10 @@ protected:
 	FRigVMGraphModifiedEvent RigVMGraphModifiedEvent;
 
 	// Delegate to subscribe to modifications to this editor data
-	UE::AnimNext::UncookedOnly::FOnEditorDataModified ModifiedDelegate;
+	UE::UAF::UncookedOnly::FOnEditorDataModified ModifiedDelegate;
 
 	// Delegate to get notified when an interaction bracket reaches 0
-	UE::AnimNext::UncookedOnly::FOnInteractionBracketFinished InteractionBracketFinished;
+	UE::UAF::UncookedOnly::FOnInteractionBracketFinished InteractionBracketFinished;
 
 	// Cached exports, generated lazily or on compilation
 	mutable TOptional<FAnimNextAssetRegistryExports> CachedExports;
diff --git a/Engine/Plugins/Experimental/AnimNext/Source/AnimNextUncookedOnly/Internal/AnimNextSharedVariableNode.h b/Engine/Plugins/Experimental/AnimNext/Source/AnimNextUncookedOnly/Internal/AnimNextSharedVariableNode.h
--- a/Engine/Plugins/Experimental/AnimNext/Source/AnimNextUncookedOnly/Internal/AnimNextSharedVariableNode.h
+++ b/Engine/Plugins/Experimental/AnimNext/Source/AnimNextUncookedOnly/Internal/AnimNextSharedVariableNode.h
@@ -9,12 +9,12 @@
 class UAnimNextSharedVariables;
 class UAnimNextControllerBase;
 
-namespace UE::AnimNext::UncookedOnly
+namespace UE::UAF::UncookedOnly
 {
 struct FUtils;
 }
 
-namespace UE::AnimNext::Editor
+namespace UE::UAF::Editor
 {
 class FAnimNextEditorModule;
 }
@@ -29,8 +29,8 @@ class UAnimNextSharedVariableNode : public URigVMVariableNode
 
 private:
 	friend UAnimNextControllerBase;
-	friend UE::AnimNext::UncookedOnly::FUtils;
-	friend UE::AnimNext::Editor::FAnimNextEditorModule;
+	friend UE::UAF::UncookedOnly::FUtils;
+	friend UE::UAF::Editor::FAnimNextEditorModule;
 
 	UPROPERTY()
 	TObjectPtr<const UAnimNextSharedVariables> Asset;
diff --git a/Engine/Plugins/Experimental/AnimNext/Source/AnimNextUncookedOnly/Internal/Entries/AnimNextCategoryEntry.h b/Engine/Plugins/Experimental/AnimNext/Source/AnimNextUncookedOnly/Internal/Entries/AnimNextCategoryEntry.h
--- a/Engine/Plugins/Experimental/AnimNext/Source/AnimNextUncookedOnly/Internal/Entries/AnimNextCategoryEntry.h
+++ b/Engine/Plugins/Experimental/AnimNext/Source/AnimNextUncookedOnly/Internal/Entries/AnimNextCategoryEntry.h
@@ -7,7 +7,7 @@
 
 #include "AnimNextCategoryEntry.generated.h"
 
-namespace UE::AnimNext::Editor
+namespace UE::UAF::Editor
 {
 struct FVariablesOutlinerCategoryItem;
 class FVariablesOutlinerHierarchy;
@@ -19,8 +19,8 @@ class UAnimNextCategoryEntry : public UAnimNextRigVMAssetEntry
 	GENERATED_BODY()
 
 	friend class UAnimNextRigVMAssetEditorData;
-	friend UE::AnimNext::Editor::FVariablesOutlinerCategoryItem;
-	friend UE::AnimNext::Editor::FVariablesOutlinerHierarchy;
+	friend UE::UAF::Editor::FVariablesOutlinerCategoryItem;
+	friend UE::UAF::Editor::FVariablesOutlinerHierarchy;
 
 	// UAnimNextRigVMAssetEntry interface
 	virtual void Initialize(UAnimNextRigVMAssetEditorData* InEditorData) override;	
diff --git a/Engine/Plugins/Experimental/AnimNext/Source/AnimNextUncookedOnly/Internal/Entries/AnimNextEventGraphEntry.h b/Engine/Plugins/Experimental/AnimNext/Source/AnimNextUncookedOnly/Internal/Entries/AnimNextEventGraphEntry.h
--- a/Engine/Plugins/Experimental/AnimNext/Source/AnimNextUncookedOnly/Internal/Entries/AnimNextEventGraphEntry.h
+++ b/Engine/Plugins/Experimental/AnimNext/Source/AnimNextUncookedOnly/Internal/Entries/AnimNextEventGraphEntry.h
@@ -13,7 +13,7 @@ class UAnimNextRigVMAssetEditorData;
 class UAnimNextModule_EditorData;
 class UAnimNextEdGraph;
 
-namespace UE::AnimNext::Tests
+namespace UE::UAF::Tests
 {
 	class FEditor_Variables;
 	class FVariables_UOLBindings;
@@ -23,7 +23,7 @@ namespace UE::AnimNext::Tests
 	class FVariables;
 }
 
-namespace UE::AnimNext::Editor
+namespace UE::UAF::Editor
 {
 	class FAnimNextEditorModule;
 }
@@ -35,13 +35,13 @@ class UAnimNextEventGraphEntry : public UAnimNextRigVMAssetEntry, public IAnimNe
 
 	friend class UAnimNextRigVMAssetEditorData;
 	friend class UAnimNextModule_EditorData;
-	friend class UE::AnimNext::Tests::FEditor_Variables;
-	friend class UE::AnimNext::Tests::FVariables_UOLBindings;
-	friend class UE::AnimNext::Tests::FEditor_AnimGraph_Variables;
-	friend class UE::AnimNext::Tests::FEditor_Graphs;
-	friend class UE::AnimNext::Tests::FEditor_AnimGraph_Graphs;
-	friend class UE::AnimNext::Tests::FVariables;
-	friend class UE::AnimNext::Editor::FAnimNextEditorModule;
+	friend class UE::UAF::Tests::FEditor_Variables;
+	friend class UE::UAF::Tests::FVariables_UOLBindings;
+	friend class UE::UAF::Tests::FEditor_AnimGraph_Variables;
+	friend class UE::UAF::Tests::FEditor_Graphs;
+	friend class UE::UAF::Tests::FEditor_AnimGraph_Graphs;
+	friend class UE::UAF::Tests::FVariables;
+	friend class UE::UAF::Editor::FAnimNextEditorModule;
 
 	// UAnimNextRigVMAssetEntry interface
 	virtual FName GetEntryName() const override { return GraphName; }
diff --git a/Engine/Plugins/Experimental/AnimNext/Source/AnimNextUncookedOnly/Internal/Entries/AnimNextRigVMAssetEntry.h b/Engine/Plugins/Experimental/AnimNext/Source/AnimNextUncookedOnly/Internal/Entries/AnimNextRigVMAssetEntry.h
--- a/Engine/Plugins/Experimental/AnimNext/Source/AnimNextUncookedOnly/Internal/Entries/AnimNextRigVMAssetEntry.h
+++ b/Engine/Plugins/Experimental/AnimNext/Source/AnimNextUncookedOnly/Internal/Entries/AnimNextRigVMAssetEntry.h
@@ -13,7 +13,7 @@ enum class ERigVMGraphNotifType : uint8;
 enum class EAnimNextEditorDataNotifType : uint8;
 class URigVMGraph;
 
-namespace UE::AnimNext::Editor
+namespace UE::UAF::Editor
 {
 	class SRigVMAssetView;
 	class SRigVMAssetViewRow;
@@ -62,8 +62,8 @@ public:
 protected:
 	UE_API void BroadcastModified(EAnimNextEditorDataNotifType InType);
 
-	friend class UE::AnimNext::Editor::FVariableCustomization;
-	friend class UE::AnimNext::Editor::SRigVMAssetView;
+	friend class UE::UAF::Editor::FVariableCustomization;
+	friend class UE::UAF::Editor::SRigVMAssetView;
 
 	UPROPERTY()
 	uint32 bHiddenInOutliner : 1 = false;
diff --git a/Engine/Plugins/Experimental/AnimNext/Source/AnimNextUncookedOnly/Internal/Entries/AnimNextSharedVariablesEntry.h b/Engine/Plugins/Experimental/AnimNext/Source/AnimNextUncookedOnly/Internal/Entries/AnimNextSharedVariablesEntry.h
--- a/Engine/Plugins/Experimental/AnimNext/Source/AnimNextUncookedOnly/Internal/Entries/AnimNextSharedVariablesEntry.h
+++ b/Engine/Plugins/Experimental/AnimNext/Source/AnimNextUncookedOnly/Internal/Entries/AnimNextSharedVariablesEntry.h
@@ -11,7 +11,7 @@ class UAnimNextSharedVariables;
 class UAnimNextRigVMAssetEditorData;
 class UAssetDefinition_AnimNextSharedVariablesEntry;
 
-namespace UE::AnimNext::Editor
+namespace UE::UAF::Editor
 {
 	class FVariablesOutlinerHierarchy;
 	struct FVariablesOutlinerStructSharedVariablesItem;
@@ -24,7 +24,7 @@ namespace UE::AnimNext::Editor
 	class SVariablesOutliner;
 }
 
-namespace UE::AnimNext::UncookedOnly
+namespace UE::UAF::UncookedOnly
 {
 	struct FUtils;
 }
@@ -44,16 +44,16 @@ class UAnimNextSharedVariablesEntry : public UAnimNextRigVMAssetEntry
 
 	friend class UAnimNextRigVMAssetEditorData;
 	friend class UAssetDefinition_AnimNextSharedVariablesEntry;
-	friend class UE::AnimNext::Editor::FVariablesOutlinerHierarchy;
-	friend struct UE::AnimNext::Editor::FVariablesOutlinerStructSharedVariablesItem;
-	friend struct UE::AnimNext::UncookedOnly::FUtils;
-	friend class UE::AnimNext::Editor::SVariablesOutlinerStructSharedVariablesLabel;
-	friend class UE::AnimNext::Editor::SVariablesOutlinerValue;
-	friend class UE::AnimNext::Editor::FVariableProxyCustomization;
-	friend class UE::AnimNext::Editor::SAddVariablesDialog;
-	friend class UE::AnimNext::Editor::SVariableOverride;
-	friend class UE::AnimNext::Editor::SVariablesOutliner;
-	friend UE::AnimNext::Editor::FVariablesOutlinerMode;
+	friend class UE::UAF::Editor::FVariablesOutlinerHierarchy;
+	friend struct UE::UAF::Editor::FVariablesOutlinerStructSharedVariablesItem;
+	friend struct UE::UAF::UncookedOnly::FUtils;
+	friend class UE::UAF::Editor::SVariablesOutlinerStructSharedVariablesLabel;
+	friend class UE::UAF::Editor::SVariablesOutlinerValue;
+	friend class UE::UAF::Editor::FVariableProxyCustomization;
+	friend class UE::UAF::Editor::SAddVariablesDialog;
+	friend class UE::UAF::Editor::SVariableOverride;
+	friend class UE::UAF::Editor::SVariablesOutliner;
+	friend UE::UAF::Editor::FVariablesOutlinerMode;
 	friend class UAnimNextAssetFindReplaceVariables;
 
 	// UAnimNextRigVMAssetEntry interface
diff --git a/Engine/Plugins/Experimental/AnimNext/Source/AnimNextUncookedOnly/Internal/Entries/AnimNextVariableEntry.h b/Engine/Plugins/Experimental/AnimNext/Source/AnimNextUncookedOnly/Internal/Entries/AnimNextVariableEntry.h
--- a/Engine/Plugins/Experimental/AnimNext/Source/AnimNextUncookedOnly/Internal/Entries/AnimNextVariableEntry.h
+++ b/Engine/Plugins/Experimental/AnimNext/Source/AnimNextUncookedOnly/Internal/Entries/AnimNextVariableEntry.h
@@ -19,7 +19,7 @@ class UAnimNextVariableBinding;
 class UAnimNextSharedVariablesEntry;
 class FAnimationAnimNextRuntimeTest_Variables;
 
-namespace UE::AnimNext::Editor
+namespace UE::UAF::Editor
 {
 	class FVariableCustomization;
 	class FVariableBindingPropertyCustomization;
@@ -35,7 +35,7 @@ namespace UE::AnimNext::Editor
 	class SVariablesOutlinerCategoryLabel;
 }
 
-namespace UE::AnimNext::Tests
+namespace UE::UAF::Tests
 {
 	class FEditor_Variables;
 	class FVariables_UOLBindings;
@@ -43,7 +43,7 @@ namespace UE::AnimNext::Tests
 	class FDataInterfaceCompile;
 }
 
-namespace UE::AnimNext::UncookedOnly
+namespace UE::UAF::UncookedOnly
 {
 	struct FPublicVariablesImpl;
 	struct FUtils;
@@ -59,25 +59,25 @@ class UAnimNextVariableEntry : public UAnimNextRigVMAssetEntry, public IAnimNext
 	friend class UAnimNextSharedVariables_EditorData;
 	friend class UAnimNextSharedVariablesEntry;
 	friend class UAnimNextStateTreeTreeEditorData;
-	friend class UE::AnimNext::Tests::FEditor_Variables;
-	friend class UE::AnimNext::Tests::FEditor_AnimGraph_Variables;
-	friend class UE::AnimNext::Tests::FVariables_UOLBindings;
+	friend class UE::UAF::Tests::FEditor_Variables;
+	friend class UE::UAF::Tests::FEditor_AnimGraph_Variables;
+	friend class UE::UAF::Tests::FVariables_UOLBindings;
 	friend class FAnimationAnimNextRuntimeTest_Variables;
-	friend class UE::AnimNext::Editor::FVariableCustomization;
-	friend class UE::AnimNext::Editor::FVariableBindingPropertyCustomization;
-	friend class UE::AnimNext::Editor::FUniversalObjectLocatorBindingCustomization;
-	friend class UE::AnimNext::Editor::FVariablesOutlinerMode;
-	friend class UE::AnimNext::Editor::SAddVariablesDialog;
-	friend class UE::AnimNext::Editor::SVariablesOutlinerValue;
-	friend class UE::AnimNext::Editor::FVariablesOutlinerHierarchy;
-	friend struct UE::AnimNext::UncookedOnly::FUtils;
-	friend class UE::AnimNext::Editor::FVariableProxyCustomization;
-	friend struct UE::AnimNext::UncookedOnly::FPublicVariablesImpl;
-	friend class UE::AnimNext::Editor::FAnimNextEditorModule;
+	friend class UE::UAF::Editor::FVariableCustomization;
+	friend class UE::UAF::Editor::FVariableBindingPropertyCustomization;
+	friend class UE::UAF::Editor::FUniversalObjectLocatorBindingCustomization;
+	friend class UE::UAF::Editor::FVariablesOutlinerMode;
+	friend class UE::UAF::Editor::SAddVariablesDialog;
+	friend class UE::UAF::Editor::SVariablesOutlinerValue;
+	friend class UE::UAF::Editor::FVariablesOutlinerHierarchy;
+	friend struct UE::UAF::UncookedOnly::FUtils;
+	friend class UE::UAF::Editor::FVariableProxyCustomization;
+	friend struct UE::UAF::UncookedOnly::FPublicVariablesImpl;
+	friend class UE::UAF::Editor::FAnimNextEditorModule;
 	friend struct FAnimNextSchemaAction_PromoteToVariable;
-	friend class UE::AnimNext::Editor::SVariablesOutlinerEntryLabel;
-	friend class UE::AnimNext::Editor::SVariablesOutlinerAssetLabel;
-	friend class UE::AnimNext::Editor::SVariablesOutlinerCategoryLabel;
+	friend class UE::UAF::Editor::SVariablesOutlinerEntryLabel;
+	friend class UE::UAF::Editor::SVariablesOutlinerAssetLabel;
+	friend class UE::UAF::Editor::SVariablesOutlinerCategoryLabel;
 	friend class UAnimNextAssetFindReplaceVariables;
 
 	UAnimNextVariableEntry();
diff --git a/Engine/Plugins/Experimental/AnimNext/Source/AnimNextUncookedOnly/Internal/Module/AnimNextModule_EditorData.h b/Engine/Plugins/Experimental/AnimNext/Source/AnimNextUncookedOnly/Internal/Module/AnimNextModule_EditorData.h
--- a/Engine/Plugins/Experimental/AnimNext/Source/AnimNextUncookedOnly/Internal/Module/AnimNextModule_EditorData.h
+++ b/Engine/Plugins/Experimental/AnimNext/Source/AnimNextUncookedOnly/Internal/Module/AnimNextModule_EditorData.h
@@ -17,12 +17,12 @@ class FAnimationAnimNextRuntimeTest_GraphAddTrait;
 class FAnimationAnimNextRuntimeTest_GraphExecute;
 class FAnimationAnimNextRuntimeTest_GraphExecuteLatent;
 
-namespace UE::AnimNext::UncookedOnly
+namespace UE::UAF::UncookedOnly
 {
 	struct FUtils;
 }
 
-namespace UE::AnimNext::Editor
+namespace UE::UAF::Editor
 {
 	class FModuleEditor;
 	class SAnimNextGraphView;
@@ -37,10 +37,10 @@ class UAnimNextModule_EditorData : public UAnimNextSharedVariables_EditorData
 
 	friend class UAnimNextModuleFactory;
 	friend class UAnimNextEdGraph;
-	friend struct UE::AnimNext::UncookedOnly::FUtils;
-	friend struct UE::AnimNext::Editor::FUtils;
-	friend class UE::AnimNext::Editor::FModuleEditor;
-	friend class UE::AnimNext::Editor::SAnimNextGraphView;
+	friend struct UE::UAF::UncookedOnly::FUtils;
+	friend struct UE::UAF::Editor::FUtils;
+	friend class UE::UAF::Editor::FModuleEditor;
+	friend class UE::UAF::Editor::SAnimNextGraphView;
 	friend struct FAnimNextGraphSchemaAction_RigUnit;
 	friend struct FAnimNextGraphSchemaAction_DispatchFactory;
 	friend class FAnimationAnimNextEditorTest_GraphAddTrait;
diff --git a/Engine/Plugins/Experimental/AnimNext/Source/AnimNextUncookedOnly/Internal/UncookedOnlyUtils.h b/Engine/Plugins/Experimental/AnimNext/Source/AnimNextUncookedOnly/Internal/UncookedOnlyUtils.h
--- a/Engine/Plugins/Experimental/AnimNext/Source/AnimNextUncookedOnly/Internal/UncookedOnlyUtils.h
+++ b/Engine/Plugins/Experimental/AnimNext/Source/AnimNextUncookedOnly/Internal/UncookedOnlyUtils.h
@@ -42,16 +42,13 @@ class UAnimNextRigVMAssetEntry;
 class URigVMUnitNode;
 class UAnimNextControllerBase;
 
-namespace UE
+namespace UE::UAF
 {
-	namespace AnimNext
-	{
-		static const FLazyName AnimNextPublicGraphFunctionsExportsRegistryTag = TEXT("AnimNextPublicGraphFunctions");
-		static const FLazyName ControlRigAssetPublicGraphFunctionsExportsRegistryTag = TEXT("PublicGraphFunctions");
-	}
+	static const FLazyName AnimNextPublicGraphFunctionsExportsRegistryTag = TEXT("AnimNextPublicGraphFunctions");
+	static const FLazyName ControlRigAssetPublicGraphFunctionsExportsRegistryTag = TEXT("PublicGraphFunctions");
 }
 
-namespace UE::AnimNext::UncookedOnly
+namespace UE::UAF::UncookedOnly
 {
 
 extern ANIMNEXTUNCOOKEDONLY_API TAutoConsoleVariable<bool> CVarDumpProgrammaticGraphs;
@@ -118,7 +115,7 @@ struct FUtils
 	{
 		const int32 StartNum = OutExports.Exports.Num();
 		FAnimNextAssetRegistryExports AssetExports;
-		const FString TagValue = InAsset.GetTagValueRef<FString>(UE::AnimNext::ExportsAnimNextAssetRegistryTag);
+		const FString TagValue = InAsset.GetTagValueRef<FString>(UE::UAF::ExportsAnimNextAssetRegistryTag);
 		if (FAnimNextAssetRegistryExports::StaticStruct()->ImportText(*TagValue, &AssetExports, nullptr, PPF_None, nullptr, FAnimNextAssetRegistryExports::StaticStruct()->GetName()) != nullptr)
 		{
 			Algo::TransformIf(AssetExports.Exports, OutExports.Exports,
@@ -139,11 +136,11 @@ struct FUtils
 	static bool GetExportsOfTypeFromAssetRegistry(TMap<FAssetData, FAnimNextAssetRegistryExports>& OutExports)
 	{
 		TArray<FAssetData> AssetData;
-		IAssetRegistry::GetChecked().GetAssetsByTags({UE::AnimNext::ExportsAnimNextAssetRegistryTag}, AssetData);
+		IAssetRegistry::GetChecked().GetAssetsByTags({UE::UAF::ExportsAnimNextAssetRegistryTag}, AssetData);
 
 		for (const FAssetData& Asset : AssetData)
 		{
-			const FString TagValue = Asset.GetTagValueRef<FString>(UE::AnimNext::ExportsAnimNextAssetRegistryTag);
+			const FString TagValue = Asset.GetTagValueRef<FString>(UE::UAF::ExportsAnimNextAssetRegistryTag);
 			FAnimNextAssetRegistryExports AssetExports;
 			if (FAnimNextAssetRegistryExports::StaticStruct()->ImportText(*TagValue, &AssetExports, nullptr, PPF_None, nullptr, FAnimNextAssetRegistryExports::StaticStruct()->GetName()) != nullptr)
 			{
diff --git a/Engine/Plugins/Experimental/AnimNext/Source/AnimNextUncookedOnly/Internal/Variables/AnimNextSharedVariables_EditorData.h b/Engine/Plugins/Experimental/AnimNext/Source/AnimNextUncookedOnly/Internal/Variables/AnimNextSharedVariables_EditorData.h
--- a/Engine/Plugins/Experimental/AnimNext/Source/AnimNextUncookedOnly/Internal/Variables/AnimNextSharedVariables_EditorData.h
+++ b/Engine/Plugins/Experimental/AnimNext/Source/AnimNextUncookedOnly/Internal/Variables/AnimNextSharedVariables_EditorData.h
@@ -9,7 +9,7 @@
 
 class UAnimNextSharedVariablesFactory;
 
-namespace UE::AnimNext::UncookedOnly
+namespace UE::UAF::UncookedOnly
 {
 	struct FUtils;
 }
@@ -27,5 +27,5 @@ protected:
 	virtual void CustomizeNewAssetEntry(UAnimNextRigVMAssetEntry* InNewEntry) const override;
 
 	friend class UAnimNextSharedVariablesFactory;
-	friend UE::AnimNext::UncookedOnly::FUtils;
+	friend UE::UAF::UncookedOnly::FUtils;
 };
diff --git a/Engine/Plugins/Experimental/AnimNext/Source/AnimNextUncookedOnly/Private/AnimNextAssetWorkspaceAssetUserData.cpp b/Engine/Plugins/Experimental/AnimNext/Source/AnimNextUncookedOnly/Private/AnimNextAssetWorkspaceAssetUserData.cpp
--- a/Engine/Plugins/Experimental/AnimNext/Source/AnimNextUncookedOnly/Private/AnimNextAssetWorkspaceAssetUserData.cpp
+++ b/Engine/Plugins/Experimental/AnimNext/Source/AnimNextUncookedOnly/Private/AnimNextAssetWorkspaceAssetUserData.cpp
@@ -19,7 +19,7 @@ void UAnimNextAssetWorkspaceAssetUserData::GetAssetRegistryTags(FAssetRegistryTa
 		CachedExports.Exports.Reset();
 		{
 			UAnimNextRigVMAsset* Asset = CastChecked<UAnimNextRigVMAsset>(GetOuter());
-			if (const UAnimNextRigVMAssetEditorData* GraphEditorData = UE::AnimNext::UncookedOnly::FUtils::GetEditorData<UAnimNextRigVMAssetEditorData>(Asset))
+			if (const UAnimNextRigVMAssetEditorData* GraphEditorData = UE::UAF::UncookedOnly::FUtils::GetEditorData<UAnimNextRigVMAssetEditorData>(Asset))
 			{
 				FWorkspaceOutlinerItemExport& RootAssetExport = CachedExports.Exports.Add_GetRef(FWorkspaceOutlinerItemExport(Asset->GetFName(), Asset));
 
@@ -38,7 +38,7 @@ void UAnimNextAssetWorkspaceAssetUserData::GetAssetRegistryTags(FAssetRegistryTa
 				FAnimNextRigVMAssetOutlinerData& Data = RootAssetExport.GetData().GetMutable<FAnimNextRigVMAssetOutlinerData>();
 				Data.SoftAssetPtr = Asset;
 
-				UE::AnimNext::UncookedOnly::FUtils::GetAssetWorkspaceExports(GraphEditorData, CachedExports, Context);
+				UE::UAF::UncookedOnly::FUtils::GetAssetWorkspaceExports(GraphEditorData, CachedExports, Context);
 			}
 		}
 	}
diff --git a/Engine/Plugins/Experimental/AnimNext/Source/AnimNextUncookedOnly/Private/AnimNextRigVMAssetEditorData.cpp b/Engine/Plugins/Experimental/AnimNext/Source/AnimNextUncookedOnly/Private/AnimNextRigVMAssetEditorData.cpp
--- a/Engine/Plugins/Experimental/AnimNext/Source/AnimNextUncookedOnly/Private/AnimNextRigVMAssetEditorData.cpp
+++ b/Engine/Plugins/Experimental/AnimNext/Source/AnimNextUncookedOnly/Private/AnimNextRigVMAssetEditorData.cpp
@@ -51,7 +51,7 @@
 
 void UAnimNextRigVMAssetEditorData::BroadcastModified(EAnimNextEditorDataNotifType InType, UObject* InSubject)
 {
-	using namespace UE::AnimNext::UncookedOnly;
+	using namespace UE::UAF::UncookedOnly;
 
 	// Modifications here can trigger compilation, so add a scope to catch compile batches
 	UAnimNextRigVMAsset* Asset = FUtils::GetAsset<UAnimNextRigVMAsset>(this);
@@ -310,17 +310,17 @@ void UAnimNextRigVMAssetEditorData::GetAssetRegistryTags(FAssetRegistryTagsConte
 
 		FString TagValue;
 		FAnimNextAssetRegistryExports::StaticStruct()->ExportText(TagValue, &CachedExports.GetValue(), nullptr, nullptr, PPF_None, nullptr);
-		Context.AddTag(FAssetRegistryTag(UE::AnimNext::ExportsAnimNextAssetRegistryTag, TagValue, FAssetRegistryTag::TT_Hidden));
+		Context.AddTag(FAssetRegistryTag(UE::UAF::ExportsAnimNextAssetRegistryTag, TagValue, FAssetRegistryTag::TT_Hidden));
 	}
 
 	{
 		FRigVMGraphFunctionHeaderArray FunctionExports;
-		UE::AnimNext::UncookedOnly::FUtils::GetAssetFunctions(this, FunctionExports);
+		UE::UAF::UncookedOnly::FUtils::GetAssetFunctions(this, FunctionExports);
 
 		FString TagValue;
 		const FArrayProperty* HeadersProperty = CastField<FArrayProperty>(FRigVMGraphFunctionHeaderArray::StaticStruct()->FindPropertyByName(TEXT("Headers")));
 		HeadersProperty->ExportText_Direct(TagValue, &(FunctionExports.Headers), &(FunctionExports.Headers), nullptr, PPF_None, nullptr);
-		Context.AddTag(FAssetRegistryTag(UE::AnimNext::AnimNextPublicGraphFunctionsExportsRegistryTag, TagValue, FAssetRegistryTag::TT_Hidden));
+		Context.AddTag(FAssetRegistryTag(UE::UAF::AnimNextPublicGraphFunctionsExportsRegistryTag, TagValue, FAssetRegistryTag::TT_Hidden));
 	}
 
 	{
@@ -373,7 +373,7 @@ void UAnimNextRigVMAssetEditorData::HandlePackageDone()
 
 void UAnimNextRigVMAssetEditorData::GetAnimNextAssetRegistryTags(FAssetRegistryTagsContext& Context, FAnimNextAssetRegistryExports& OutExports) const
 {	
-	UE::AnimNext::UncookedOnly::FUtils::GetAssetVariableExports(this, OutExports, Context);
+	UE::UAF::UncookedOnly::FUtils::GetAssetVariableExports(this, OutExports, Context);
 }
 
 void UAnimNextRigVMAssetEditorData::RefreshAllModels(ERigVMLoadType InLoadType)
@@ -798,7 +798,7 @@ TObjectPtr<URigVMGraph> UAnimNextRigVMAssetEditorData::CreateContainedGraphModel
 
 void UAnimNextRigVMAssetEditorData::BuildFunctionWrapperEventVariables(FAnimNextRigVMAssetCompileContext& InContext)
 {
-	using namespace UE::AnimNext::UncookedOnly;
+	using namespace UE::UAF::UncookedOnly;
 	
 	for (const FAnimNextProgrammaticFunctionHeader& ProgrammaticFunctionHeader : InContext.FunctionHeaders)
 	{
@@ -818,7 +818,7 @@ void UAnimNextRigVMAssetEditorData::BuildFunctionWrapperEventVariables(FAnimNext
 
 void UAnimNextRigVMAssetEditorData::BuildFunctionWrapperEvents(FAnimNextRigVMAssetCompileContext& InContext, const FRigVMCompileSettings& InSettings)
 {
-	using namespace UE::AnimNext::UncookedOnly;
+	using namespace UE::UAF::UncookedOnly;
 
 	UAnimNextRigVMAsset* Asset = FUtils::GetAsset<UAnimNextRigVMAsset>(this);
 	Asset->FunctionData.Reset();
@@ -987,7 +987,7 @@ void UAnimNextRigVMAssetEditorData::BuildFunctionWrapperEvents(FAnimNextRigVMAss
 
 void UAnimNextRigVMAssetEditorData::RecompileVM()
 {
-	using namespace UE::AnimNext::UncookedOnly;
+	using namespace UE::UAF::UncookedOnly;
 
 	if (bIsCompiling)
 	{
@@ -1150,7 +1150,7 @@ void UAnimNextRigVMAssetEditorData::HandleRemoveNotify(UObject* InAsset, const F
 		return;
 	}
 
-	UAnimNextRigVMAssetEditorData* EditorData = UE::AnimNext::UncookedOnly::FUtils::GetEditorData<UAnimNextRigVMAssetEditorData>(Asset);
+	UAnimNextRigVMAssetEditorData* EditorData = UE::UAF::UncookedOnly::FUtils::GetEditorData<UAnimNextRigVMAssetEditorData>(Asset);
 	if(EditorData == nullptr)
 	{
 		return;
@@ -1190,7 +1190,7 @@ void UAnimNextRigVMAssetEditorData::HandleReplaceNotify(UObject* InAsset, const
 		return;
 	}
 
-	UAnimNextRigVMAssetEditorData* EditorData = UE::AnimNext::UncookedOnly::FUtils::GetEditorData<UAnimNextRigVMAssetEditorData>(Asset);
+	UAnimNextRigVMAssetEditorData* EditorData = UE::UAF::UncookedOnly::FUtils::GetEditorData<UAnimNextRigVMAssetEditorData>(Asset);
 	if(EditorData == nullptr)
 	{
 		return;
@@ -1235,7 +1235,7 @@ bool UAnimNextRigVMAssetEditorData::IsDirtyForRecompilation() const
 	{
 		if(InEntry->Asset)
 		{
-			UAnimNextRigVMAssetEditorData* EditorData = UE::AnimNext::UncookedOnly::FUtils::GetEditorData<UAnimNextRigVMAssetEditorData>(InEntry->Asset.Get());
+			UAnimNextRigVMAssetEditorData* EditorData = UE::UAF::UncookedOnly::FUtils::GetEditorData<UAnimNextRigVMAssetEditorData>(InEntry->Asset.Get());
 			if(EditorData->IsDirtyForRecompilation())
 			{
 				bDependencyDirty = true;
@@ -1451,7 +1451,7 @@ TArray<UEdGraph*> UAnimNextRigVMAssetEditorData::GetAllEdGraphs() const
 
 UAnimNextRigVMAssetEntry* UAnimNextRigVMAssetLibrary::FindEntry(UAnimNextRigVMAsset* InAsset, FName InName)
 {
-	return UE::AnimNext::UncookedOnly::FUtils::GetEditorData(InAsset)->FindEntry(InName);
+	return UE::UAF::UncookedOnly::FUtils::GetEditorData(InAsset)->FindEntry(InName);
 }
 
 UAnimNextRigVMAssetEntry* UAnimNextRigVMAssetEditorData::FindEntry(FName InName) const
@@ -1642,7 +1642,7 @@ void UAnimNextRigVMAssetEditorData::ReorderVariable(UAnimNextVariableEntry* Vari
 
 bool UAnimNextRigVMAssetLibrary::RemoveEntry(UAnimNextRigVMAsset* InAsset, UAnimNextRigVMAssetEntry* InEntry,  bool bSetupUndoRedo, bool bPrintPythonCommand)
 {
-	return UE::AnimNext::UncookedOnly::FUtils::GetEditorData(InAsset)->RemoveEntry(InEntry, bSetupUndoRedo, bPrintPythonCommand);
+	return UE::UAF::UncookedOnly::FUtils::GetEditorData(InAsset)->RemoveEntry(InEntry, bSetupUndoRedo, bPrintPythonCommand);
 }
 
 bool UAnimNextRigVMAssetEditorData::RemoveEntry(UAnimNextRigVMAssetEntry* InEntry, bool bSetupUndoRedo, bool bPrintPythonCommand)
@@ -1704,7 +1704,7 @@ bool UAnimNextRigVMAssetEditorData::RemoveEntry(UAnimNextRigVMAssetEntry* InEntr
 
 bool UAnimNextRigVMAssetLibrary::RemoveEntries(UAnimNextRigVMAsset* InAsset, const TArray<UAnimNextRigVMAssetEntry*>& InEntries,  bool bSetupUndoRedo, bool bPrintPythonCommand)
 {
-	return UE::AnimNext::UncookedOnly::FUtils::GetEditorData(InAsset)->RemoveEntries(InEntries, bSetupUndoRedo, bPrintPythonCommand);
+	return UE::UAF::UncookedOnly::FUtils::GetEditorData(InAsset)->RemoveEntries(InEntries, bSetupUndoRedo, bPrintPythonCommand);
 }
 
 bool UAnimNextRigVMAssetEditorData::RemoveEntries(TConstArrayView<UAnimNextRigVMAssetEntry*> InEntries, bool bSetupUndoRedo, bool bPrintPythonCommand)
@@ -1746,7 +1746,7 @@ bool UAnimNextRigVMAssetEditorData::RemoveEntries(TConstArrayView<UAnimNextRigVM
 
 bool UAnimNextRigVMAssetLibrary::RemoveAllEntries(UAnimNextRigVMAsset* InAsset, bool bSetupUndoRedo, bool bPrintPythonCommand)
 {
-	return UE::AnimNext::UncookedOnly::FUtils::GetEditorData(InAsset)->RemoveAllEntries(bSetupUndoRedo, bPrintPythonCommand);
+	return UE::UAF::UncookedOnly::FUtils::GetEditorData(InAsset)->RemoveAllEntries(bSetupUndoRedo, bPrintPythonCommand);
 }
 
 bool UAnimNextRigVMAssetEditorData::RemoveAllEntries(bool bSetupUndoRedo, bool bPrintPythonCommand)
@@ -1778,7 +1778,7 @@ UObject* UAnimNextRigVMAssetEditorData::CreateNewSubEntry(UAnimNextRigVMAssetEdi
 {
 	UObject* NewEntry = NewObject<UObject>(InEditorData, InClass.Get(), NAME_None, RF_Transactional);
 	// If we are a transient asset, dont use external packages
-	UAnimNextRigVMAsset* Asset = UE::AnimNext::UncookedOnly::FUtils::GetAsset(InEditorData);
+	UAnimNextRigVMAsset* Asset = UE::UAF::UncookedOnly::FUtils::GetAsset(InEditorData);
 	check(Asset);
 
 	// Additionally check external packaging flag
@@ -2073,7 +2073,7 @@ bool UAnimNextRigVMAssetEditorData::RemoveEdGraph(URigVMGraph* InModel)
 UAnimNextVariableEntry* UAnimNextRigVMAssetLibrary::AddVariable(UAnimNextRigVMAsset* InAsset, FName InName, EPropertyBagPropertyType InValueType,
 	EPropertyBagContainerType InContainerType, const UObject* InValueTypeObject, const FString& InDefaultValue, bool bSetupUndoRedo, bool bPrintPythonCommand)
 {
-	return UE::AnimNext::UncookedOnly::FUtils::GetEditorData(InAsset)->AddVariable(InName, FAnimNextParamType(InValueType, InContainerType, InValueTypeObject), InDefaultValue, bSetupUndoRedo, bPrintPythonCommand);
+	return UE::UAF::UncookedOnly::FUtils::GetEditorData(InAsset)->AddVariable(InName, FAnimNextParamType(InValueType, InContainerType, InValueTypeObject), InDefaultValue, bSetupUndoRedo, bPrintPythonCommand);
 }
 
 UAnimNextVariableEntry* UAnimNextRigVMAssetEditorData::AddVariable(FName InName, FAnimNextParamType InType, const FString& InDefaultValue, bool bSetupUndoRedo, bool bPrintPythonCommand)
@@ -2156,7 +2156,7 @@ UAnimNextVariableEntry* UAnimNextRigVMAssetEditorData::AddVariable(FName InName,
 
 UAnimNextEventGraphEntry* UAnimNextRigVMAssetLibrary::AddEventGraph(UAnimNextRigVMAsset* InAsset, FName InName, UScriptStruct* InEventStruct, bool bSetupUndoRedo, bool bPrintPythonCommand)
 {
-	return UE::AnimNext::UncookedOnly::FUtils::GetEditorData(InAsset)->AddEventGraph(InName, InEventStruct, bSetupUndoRedo, bPrintPythonCommand);
+	return UE::UAF::UncookedOnly::FUtils::GetEditorData(InAsset)->AddEventGraph(InName, InEventStruct, bSetupUndoRedo, bPrintPythonCommand);
 }
 
 UAnimNextEventGraphEntry* UAnimNextRigVMAssetEditorData::AddEventGraph(FName InName, UScriptStruct* InEventStruct, bool bSetupUndoRedo, bool bPrintPythonCommand)
@@ -2225,7 +2225,7 @@ UAnimNextEventGraphEntry* UAnimNextRigVMAssetEditorData::AddEventGraph(FName InN
 			RefreshExternalModels();
 			RigVMClient.AddModel(NewRigVMGraphModel, true);
 			URigVMController* Controller = RigVMClient.GetController(NewRigVMGraphModel);
-			UE::AnimNext::UncookedOnly::FUtils::SetupEventGraph(Controller, InEventStruct, NewGraphName);
+			UE::UAF::UncookedOnly::FUtils::SetupEventGraph(Controller, InEventStruct, NewGraphName);
 		}
 	}
 
@@ -2245,7 +2245,7 @@ UAnimNextEventGraphEntry* UAnimNextRigVMAssetEditorData::AddEventGraph(FName InN
 
 UAnimNextSharedVariablesEntry* UAnimNextRigVMAssetLibrary::AddSharedVariables(UAnimNextRigVMAsset* InAsset, UAnimNextSharedVariables* InSharedVariables, bool bSetupUndoRedo, bool bPrintPythonCommand)
 {
-	return UE::AnimNext::UncookedOnly::FUtils::GetEditorData(InAsset)->AddSharedVariables(InSharedVariables, bSetupUndoRedo, bPrintPythonCommand);
+	return UE::UAF::UncookedOnly::FUtils::GetEditorData(InAsset)->AddSharedVariables(InSharedVariables, bSetupUndoRedo, bPrintPythonCommand);
 }
 
 UAnimNextSharedVariablesEntry* UAnimNextRigVMAssetEditorData::AddSharedVariables(const UAnimNextSharedVariables* InSharedVariables, bool bSetupUndoRedo, bool bPrintPythonCommand)
@@ -2263,7 +2263,7 @@ UAnimNextSharedVariablesEntry* UAnimNextRigVMAssetEditorData::AddSharedVariables
 	}
 	
 	// Check if interface has any public members or if any of its parent interfaces do
-	UAnimNextSharedVariables_EditorData* EditorData = UE::AnimNext::UncookedOnly::FUtils::GetEditorData<UAnimNextSharedVariables_EditorData>(InSharedVariables);
+	UAnimNextSharedVariables_EditorData* EditorData = UE::UAF::UncookedOnly::FUtils::GetEditorData<UAnimNextSharedVariables_EditorData>(InSharedVariables);
 	if(EditorData == nullptr)
 	{
 		ReportError(TEXT("UAnimNextRigVMAssetEditorData::AddSharedVariables: Invalid asset supplied - asset has no editor data."));
@@ -2283,7 +2283,7 @@ UAnimNextSharedVariablesEntry* UAnimNextRigVMAssetEditorData::AddSharedVariables
 			if(UAnimNextSharedVariablesEntry* SharedVariablesEntry = Cast<UAnimNextSharedVariablesEntry>(Entry))
 			{
 				const UAnimNextSharedVariables* SharedVariables = SharedVariablesEntry->GetAsset();
-				UAnimNextSharedVariables_EditorData* EditorData = UE::AnimNext::UncookedOnly::FUtils::GetEditorData<UAnimNextSharedVariables_EditorData>(SharedVariables);
+				UAnimNextSharedVariables_EditorData* EditorData = UE::UAF::UncookedOnly::FUtils::GetEditorData<UAnimNextSharedVariables_EditorData>(SharedVariables);
 				if(InCheckForCircularity(EditorData, InCheckForCircularity))
 				{
 					return true;
@@ -2314,7 +2314,7 @@ UAnimNextSharedVariablesEntry* UAnimNextRigVMAssetEditorData::AddSharedVariables
 			else if(UAnimNextSharedVariablesEntry* SharedVariablesEntry = Cast<UAnimNextSharedVariablesEntry>(Entry))
 			{
 				const UAnimNextSharedVariables* SharedVariables = SharedVariablesEntry->GetAsset();
-				UAnimNextSharedVariables_EditorData* EditorData = UE::AnimNext::UncookedOnly::FUtils::GetEditorData<UAnimNextSharedVariables_EditorData>(SharedVariables);
+				UAnimNextSharedVariables_EditorData* EditorData = UE::UAF::UncookedOnly::FUtils::GetEditorData<UAnimNextSharedVariables_EditorData>(SharedVariables);
 				if(InCheckForPublicMembers(EditorData, InCheckForPublicMembers))
 				{
 					return true;
@@ -2376,7 +2376,7 @@ UAnimNextSharedVariablesEntry* UAnimNextRigVMAssetEditorData::AddSharedVariables
 
 UAnimNextSharedVariablesEntry* UAnimNextRigVMAssetLibrary::AddSharedVariablesStruct(UAnimNextRigVMAsset* InAsset, UScriptStruct* InStruct, bool bSetupUndoRedo, bool bPrintPythonCommand)
 {
-	return UE::AnimNext::UncookedOnly::FUtils::GetEditorData(InAsset)->AddSharedVariablesStruct(InStruct, bSetupUndoRedo, bPrintPythonCommand);
+	return UE::UAF::UncookedOnly::FUtils::GetEditorData(InAsset)->AddSharedVariablesStruct(InStruct, bSetupUndoRedo, bPrintPythonCommand);
 }
 
 UAnimNextSharedVariablesEntry* UAnimNextRigVMAssetEditorData::AddSharedVariablesStruct(const UScriptStruct* InStruct, bool bSetupUndoRedo, bool bPrintPythonCommand)
@@ -2437,17 +2437,17 @@ UAnimNextSharedVariablesEntry* UAnimNextRigVMAssetEditorData::AddSharedVariables
 
 URigVMLibraryNode* UAnimNextRigVMAssetLibrary::AddFunction(UAnimNextRigVMAsset* InAsset, FName InFunctionName, bool bInMutable, bool bSetupUndoRedo, bool bPrintPythonCommand)
 {
-	return UE::AnimNext::UncookedOnly::FUtils::GetEditorData(InAsset)->AddFunction(InFunctionName, bInMutable, bSetupUndoRedo, bPrintPythonCommand);
+	return UE::UAF::UncookedOnly::FUtils::GetEditorData(InAsset)->AddFunction(InFunctionName, bInMutable, bSetupUndoRedo, bPrintPythonCommand);
 }
 
 bool UAnimNextRigVMAssetLibrary::AddCategory(UAnimNextRigVMAsset* InAsset, const FString& CategoryName, bool bSetupUndoRedo, bool bPrintPythonCommand)
 {
-	return UE::AnimNext::UncookedOnly::FUtils::GetEditorData(InAsset)->AddCategory(CategoryName, bSetupUndoRedo, bPrintPythonCommand);
+	return UE::UAF::UncookedOnly::FUtils::GetEditorData(InAsset)->AddCategory(CategoryName, bSetupUndoRedo, bPrintPythonCommand);
 }
 
 bool UAnimNextRigVMAssetLibrary::RenameCategory(UAnimNextRigVMAsset* InAsset, const FString& CategoryName, const FString& NewCategoryName, bool bSetupUndoRedo, bool bPrintPythonCommand)
 {
-	return UE::AnimNext::UncookedOnly::FUtils::GetEditorData(InAsset)->RenameCategory(CategoryName, NewCategoryName, bSetupUndoRedo, bPrintPythonCommand);
+	return UE::UAF::UncookedOnly::FUtils::GetEditorData(InAsset)->RenameCategory(CategoryName, NewCategoryName, bSetupUndoRedo, bPrintPythonCommand);
 }
 
 URigVMLibraryNode* UAnimNextRigVMAssetEditorData::AddFunction(FName InFunctionName, bool bInMutable, bool bInSetupUndoRedo, bool bPrintPythonCommand)
@@ -2486,7 +2486,7 @@ bool UAnimNextRigVMAssetEditorData::HasPublicVariables() const
 		{
 			if(SharedVariablesEntry->GetAsset())
 			{
-				UAnimNextSharedVariables_EditorData* EditorData = UE::AnimNext::UncookedOnly::FUtils::GetEditorData<UAnimNextSharedVariables_EditorData>(SharedVariablesEntry->GetAsset());
+				UAnimNextSharedVariables_EditorData* EditorData = UE::UAF::UncookedOnly::FUtils::GetEditorData<UAnimNextSharedVariables_EditorData>(SharedVariablesEntry->GetAsset());
 				return EditorData->HasPublicVariables();
 			}
 		}
@@ -2503,7 +2503,7 @@ void UAnimNextRigVMAssetEditorData::GetAllVariables(TArray<FVariableInfo>& OutVa
 			if (InAccess == EVariableAccessFilter::All ||
 				VariableEntry->GetExportAccessSpecifier() == EAnimNextExportAccessSpecifier::Public)
 			{
-				const UAnimNextRigVMAsset* Asset = UE::AnimNext::UncookedOnly::FUtils::GetAsset<UAnimNextRigVMAsset>(this);
+				const UAnimNextRigVMAsset* Asset = UE::UAF::UncookedOnly::FUtils::GetAsset<UAnimNextRigVMAsset>(this);
 
 				const FProperty* Property;
 				TConstArrayView<uint8> Value;
@@ -2532,7 +2532,7 @@ void UAnimNextRigVMAssetEditorData::GetAllVariables(TArray<FVariableInfo>& OutVa
 			case EAnimNextSharedVariablesType::Asset:
 				if(SharedVariablesEntry->GetAsset())
 				{
-					UAnimNextSharedVariables_EditorData* EditorData = UE::AnimNext::UncookedOnly::FUtils::GetEditorData<UAnimNextSharedVariables_EditorData>(SharedVariablesEntry->GetAsset());
+					UAnimNextSharedVariables_EditorData* EditorData = UE::UAF::UncookedOnly::FUtils::GetEditorData<UAnimNextSharedVariables_EditorData>(SharedVariablesEntry->GetAsset());
 					EditorData->GetAllVariables(OutVariables, InRecursion, EVariableAccessFilter::PublicOnly);
 				}
 				break;
@@ -2569,7 +2569,7 @@ void UAnimNextRigVMAssetEditorData::HandleReportFromCompiler(EMessageSeverity::T
 		return;
 	}
 
-	FCompilerResultsLog& Log = UE::AnimNext::UncookedOnly::FScopedCompileJob::GetLog();
+	FCompilerResultsLog& Log = UE::UAF::UncookedOnly::FScopedCompileJob::GetLog();
 
 	UObject* SubjectForMessage = InSubject;
 	if(URigVMNode* ModelNode = Cast<URigVMNode>(SubjectForMessage))
@@ -2752,7 +2752,7 @@ void UAnimNextRigVMAssetEditorData::SetUseExternalPackages(TArrayView<UAnimNextR
 			continue;
 		}
 
-		UAnimNextRigVMAssetEditorData* EditorData = UE::AnimNext::UncookedOnly::FUtils::GetEditorData<UAnimNextRigVMAssetEditorData>(Asset);
+		UAnimNextRigVMAssetEditorData* EditorData = UE::UAF::UncookedOnly::FUtils::GetEditorData<UAnimNextRigVMAssetEditorData>(Asset);
 		if(EditorData == nullptr)
 		{
 			continue;
@@ -2909,7 +2909,7 @@ void UAnimNextRigVMAssetEditorData::SetUseExternalPackages(TArrayView<UAnimNextR
 
 FInstancedPropertyBag UAnimNextRigVMAssetEditorData::GenerateCombinedPropertyBag(const FAnimNextGetVariableCompileContext& InCompileContext) const
 {
-	UAnimNextRigVMAsset* Asset = UE::AnimNext::UncookedOnly::FUtils::GetAsset<UAnimNextRigVMAsset>(this);
+	UAnimNextRigVMAsset* Asset = UE::UAF::UncookedOnly::FUtils::GetAsset<UAnimNextRigVMAsset>(this);
 
 	FInstancedPropertyBag CombinedPropertyBag;
 
@@ -2929,7 +2929,7 @@ FInstancedPropertyBag UAnimNextRigVMAssetEditorData::GenerateCombinedPropertyBag
 	// It must match that of the variables FUAFInstanceVariableData::Initialize
 
 	// First add internal variables
-	FInstancedPropertyBag AssetBag = UE::AnimNext::UncookedOnly::FUtils::MakePropertyBagForEditorData(this, InCompileContext);
+	FInstancedPropertyBag AssetBag = UE::UAF::UncookedOnly::FUtils::MakePropertyBagForEditorData(this, InCompileContext);
 	AddInstancedPropertyBag(AssetBag);
 
 	// Next add shared variables
@@ -2940,16 +2940,16 @@ FInstancedPropertyBag UAnimNextRigVMAssetEditorData::GenerateCombinedPropertyBag
 			continue;
 		}
 
-		const UAnimNextRigVMAssetEditorData* ReferencedEditorData = UE::AnimNext::UncookedOnly::FUtils::GetEditorData<UAnimNextRigVMAssetEditorData>(ReferencedVariableAsset);
-		AssetBag = UE::AnimNext::UncookedOnly::FUtils::MakePropertyBagForEditorData(ReferencedEditorData, InCompileContext);
+		const UAnimNextRigVMAssetEditorData* ReferencedEditorData = UE::UAF::UncookedOnly::FUtils::GetEditorData<UAnimNextRigVMAssetEditorData>(ReferencedVariableAsset);
+		AssetBag = UE::UAF::UncookedOnly::FUtils::MakePropertyBagForEditorData(ReferencedEditorData, InCompileContext);
 		AddInstancedPropertyBag(AssetBag);
 	}
 
 	// Next add native structs
 	for (const UScriptStruct* ReferencedStruct : Asset->ReferencedVariableStructs)
 	{
-		TSharedRef<UE::AnimNext::FStructDataCache> StructData = UE::AnimNext::FStructDataCache::GetStructInfo(ReferencedStruct);
-		for (const UE::AnimNext::FStructDataCache::FPropertyInfo& PropertyInfo : StructData->GetProperties())
+		TSharedRef<UE::UAF::FStructDataCache> StructData = UE::UAF::FStructDataCache::GetStructInfo(ReferencedStruct);
+		for (const UE::UAF::FStructDataCache::FPropertyInfo& PropertyInfo : StructData->GetProperties())
 		{
 			CombinedPropertyBag.AddProperty(PropertyInfo.Property->GetFName(), PropertyInfo.Property);
 		}
diff --git a/Engine/Plugins/Experimental/AnimNext/Source/AnimNextUncookedOnly/Private/AnimNextScopedCompileJob.cpp b/Engine/Plugins/Experimental/AnimNext/Source/AnimNextUncookedOnly/Private/AnimNextScopedCompileJob.cpp
--- a/Engine/Plugins/Experimental/AnimNext/Source/AnimNextUncookedOnly/Private/AnimNextScopedCompileJob.cpp
+++ b/Engine/Plugins/Experimental/AnimNext/Source/AnimNextUncookedOnly/Private/AnimNextScopedCompileJob.cpp
@@ -10,7 +10,7 @@
 
 #define LOCTEXT_NAMESPACE "AnimNextScopedCompileJob"
 
-namespace UE::AnimNext::UncookedOnly
+namespace UE::UAF::UncookedOnly
 {
 
 struct FCompileJobThreadData
diff --git a/Engine/Plugins/Experimental/AnimNext/Source/AnimNextUncookedOnly/Private/AnimNextUncookedOnlyModule.cpp b/Engine/Plugins/Experimental/AnimNext/Source/AnimNextUncookedOnly/Private/AnimNextUncookedOnlyModule.cpp
--- a/Engine/Plugins/Experimental/AnimNext/Source/AnimNextUncookedOnly/Private/AnimNextUncookedOnlyModule.cpp
+++ b/Engine/Plugins/Experimental/AnimNext/Source/AnimNextUncookedOnly/Private/AnimNextUncookedOnlyModule.cpp
@@ -12,7 +12,7 @@
 
 #define LOCTEXT_NAMESPACE "AnimNextUncookedOnlyModule"
 
-namespace UE::AnimNext::UncookedOnly
+namespace UE::UAF::UncookedOnly
 {
 
 void FModule::StartupModule()
@@ -56,4 +56,4 @@ TSharedPtr<IVariableBindingType> FModule::FindVariableBindingType(const UScriptS
 
 #undef LOCTEXT_NAMESPACE 
 
-IMPLEMENT_MODULE(UE::AnimNext::UncookedOnly::FModule, AnimNextUncookedOnly);
+IMPLEMENT_MODULE(UE::UAF::UncookedOnly::FModule, AnimNextUncookedOnly);
diff --git a/Engine/Plugins/Experimental/AnimNext/Source/AnimNextUncookedOnly/Private/AnimNextUncookedOnlyModule.h b/Engine/Plugins/Experimental/AnimNext/Source/AnimNextUncookedOnly/Private/AnimNextUncookedOnlyModule.h
--- a/Engine/Plugins/Experimental/AnimNext/Source/AnimNextUncookedOnly/Private/AnimNextUncookedOnlyModule.h
+++ b/Engine/Plugins/Experimental/AnimNext/Source/AnimNextUncookedOnly/Private/AnimNextUncookedOnlyModule.h
@@ -5,12 +5,12 @@
 #include "CoreMinimal.h"
 #include "IAnimNextUncookedOnlyModule.h"
 
-namespace UE::AnimNext::UncookedOnly
+namespace UE::UAF::UncookedOnly
 {
 	struct FUtils;
 }
 
-namespace UE::AnimNext::UncookedOnly
+namespace UE::UAF::UncookedOnly
 {
 
 class FModule : public IAnimNextUncookedOnlyModule
diff --git a/Engine/Plugins/Experimental/AnimNext/Source/AnimNextUncookedOnly/Private/Entries/AnimNextSharedVariablesEntry.cpp b/Engine/Plugins/Experimental/AnimNext/Source/AnimNextUncookedOnly/Private/Entries/AnimNextSharedVariablesEntry.cpp
--- a/Engine/Plugins/Experimental/AnimNext/Source/AnimNextUncookedOnly/Private/Entries/AnimNextSharedVariablesEntry.cpp
+++ b/Engine/Plugins/Experimental/AnimNext/Source/AnimNextUncookedOnly/Private/Entries/AnimNextSharedVariablesEntry.cpp
@@ -21,7 +21,7 @@ void UAnimNextSharedVariablesEntry::Initialize(UAnimNextRigVMAssetEditorData* In
 
 	if(Asset)
 	{
-		UAnimNextSharedVariables_EditorData* EditorData = UE::AnimNext::UncookedOnly::FUtils::GetEditorData<UAnimNextSharedVariables_EditorData>(Asset.Get());
+		UAnimNextSharedVariables_EditorData* EditorData = UE::UAF::UncookedOnly::FUtils::GetEditorData<UAnimNextSharedVariables_EditorData>(Asset.Get());
 		EditorData->ModifiedDelegate.AddUObject(this, &UAnimNextSharedVariablesEntry::HandleAssetModified);
 	}
 	else if (Struct)
diff --git a/Engine/Plugins/Experimental/AnimNext/Source/AnimNextUncookedOnly/Private/Module/AnimNextModule_EditorData.cpp b/Engine/Plugins/Experimental/AnimNext/Source/AnimNextUncookedOnly/Private/Module/AnimNextModule_EditorData.cpp
--- a/Engine/Plugins/Experimental/AnimNext/Source/AnimNextUncookedOnly/Private/Module/AnimNextModule_EditorData.cpp
+++ b/Engine/Plugins/Experimental/AnimNext/Source/AnimNextUncookedOnly/Private/Module/AnimNextModule_EditorData.cpp
@@ -29,7 +29,7 @@ TConstArrayView<TSubclassOf<UAnimNextRigVMAssetEntry>> UAnimNextModule_EditorDat
 
 void UAnimNextModule_EditorData::OnPreCompileAsset(FRigVMCompileSettings& InSettings)
 {
-	using namespace UE::AnimNext::UncookedOnly;
+	using namespace UE::UAF::UncookedOnly;
 
 	UAnimNextModule* Module = FUtils::GetAsset<UAnimNextModule>(this);
 
@@ -39,14 +39,14 @@ void UAnimNextModule_EditorData::OnPreCompileAsset(FRigVMCompileSettings& InSett
 
 void UAnimNextModule_EditorData::OnPreCompileGetProgrammaticGraphs(const FRigVMCompileSettings& InSettings, FAnimNextGetGraphCompileContext& OutCompileContext)
 {
-	using namespace UE::AnimNext::UncookedOnly;
+	using namespace UE::UAF::UncookedOnly;
 
 	FUtils::CompileVariableBindings(InSettings, FUtils::GetAsset<UAnimNextModule>(this), OutCompileContext.GetMutableProgrammaticGraphs());
 }
 
 void UAnimNextModule_EditorData::OnPreCompileProcessGraphs(const FRigVMCompileSettings& InSettings, FAnimNextProcessGraphCompileContext& OutCompileContext)
 {
-	using namespace UE::AnimNext::UncookedOnly;
+	using namespace UE::UAF::UncookedOnly;
 
 	UAnimNextModule* Module = FUtils::GetAsset<UAnimNextModule>(this);
 
diff --git a/Engine/Plugins/Experimental/AnimNext/Source/AnimNextUncookedOnly/Private/UncookedOnlyUtils.cpp b/Engine/Plugins/Experimental/AnimNext/Source/AnimNextUncookedOnly/Private/UncookedOnlyUtils.cpp
--- a/Engine/Plugins/Experimental/AnimNext/Source/AnimNextUncookedOnly/Private/UncookedOnlyUtils.cpp
+++ b/Engine/Plugins/Experimental/AnimNext/Source/AnimNextUncookedOnly/Private/UncookedOnlyUtils.cpp
@@ -39,7 +39,7 @@
 
 #define LOCTEXT_NAMESPACE "AnimNextUncookedOnlyUtils"
 
-namespace UE::AnimNext::UncookedOnly
+namespace UE::UAF::UncookedOnly
 {
 
 TAutoConsoleVariable<bool> CVarDumpProgrammaticGraphs(
@@ -1644,7 +1644,7 @@ void FUtils::ReplaceVariableReferences(UAnimNextRigVMAssetEditorData* InEditorDa
 
 		if (!bHasEntry)
 		{
-			UE::AnimNext::UncookedOnly::FScopedCompileJob CompilerResults(LOCTEXT("AddSharedAssetJobName", "Add Shared Asset dependency"), { ThisOuterAsset });
+			UE::UAF::UncookedOnly::FScopedCompileJob CompilerResults(LOCTEXT("AddSharedAssetJobName", "Add Shared Asset dependency"), { ThisOuterAsset });
 			InEditorData->AddSharedVariables(SharedVariablesReplaceOuter, bSetupUndoRedo, bPrintPythonCommands);
 			InEditorData->RecompileVM();
 		}
@@ -1775,7 +1775,7 @@ void FUtils::ReplaceVariableReferencesAcrossProject(const FAnimNextSoftVariableR
 	AssetFilter.bRecursiveClasses = true;
 
 	// [TODO] would this miss out assets
-	//AssetFilter.TagsAndValues.Add(UE::AnimNext::ExportsAnimNextAssetRegistryTag, TOptional<FString>());
+	//AssetFilter.TagsAndValues.Add(UE::UAF::ExportsAnimNextAssetRegistryTag, TOptional<FString>());
 		
 	TArray<FAssetData> AssetDatas;
 	FModuleManager::GetModuleChecked<FAssetRegistryModule>("AssetRegistry").Get().GetAssets(AssetFilter, AssetDatas);
@@ -1784,7 +1784,7 @@ void FUtils::ReplaceVariableReferencesAcrossProject(const FAnimNextSoftVariableR
 	for (const FAssetData& AssetData : AssetDatas)
 	{
 		// Relates to [TODO] above
-		ensure(AssetData.FindTag(UE::AnimNext::ExportsAnimNextAssetRegistryTag));
+		ensure(AssetData.FindTag(UE::UAF::ExportsAnimNextAssetRegistryTag));
 		
 		FAnimNextAssetRegistryExports AnimNextExports;
 		GetExportsOfTypeForAsset<FAnimNextVariableReferenceData>(AssetData, AnimNextExports);
diff --git a/Engine/Plugins/Experimental/AnimNext/Source/AnimNextUncookedOnly/Private/Variables/K2Node_AnimNextComponentSetVariable.cpp b/Engine/Plugins/Experimental/AnimNext/Source/AnimNextUncookedOnly/Private/Variables/K2Node_AnimNextComponentSetVariable.cpp
--- a/Engine/Plugins/Experimental/AnimNext/Source/AnimNextUncookedOnly/Private/Variables/K2Node_AnimNextComponentSetVariable.cpp
+++ b/Engine/Plugins/Experimental/AnimNext/Source/AnimNextUncookedOnly/Private/Variables/K2Node_AnimNextComponentSetVariable.cpp
@@ -49,10 +49,10 @@ void UK2Node_AnimNextComponentSetVariable::CachePinType()
 		FAnimNextVariableReference::StaticStruct()->ImportText(*VariablePin->DefaultValue, &VariableReference, nullptr, PPF_None, nullptr, FAnimNextVariableReference::StaticStruct()->GetName());
 		if (VariableReference.IsValid())
 		{
-			FAnimNextParamType Type = UE::AnimNext::UncookedOnly::FUtils::FindVariableType(VariableReference);
+			FAnimNextParamType Type = UE::UAF::UncookedOnly::FUtils::FindVariableType(VariableReference);
 			if (Type.IsValid())
 			{
-				CachedPinType = UE::AnimNext::UncookedOnly::FUtils::GetPinTypeFromParamType(Type);
+				CachedPinType = UE::UAF::UncookedOnly::FUtils::GetPinTypeFromParamType(Type);
 			}
 		}
 	}
@@ -97,10 +97,10 @@ void UK2Node_AnimNextComponentSetVariable::ValidateNodeDuringCompilation(FCompil
 	FAnimNextVariableReference::StaticStruct()->ImportText(*VariablePin->DefaultValue, &VariableReference, nullptr, PPF_None, nullptr, FAnimNextVariableReference::StaticStruct()->GetName());
 	if (VariableReference.IsValid())
 	{
-		FAnimNextParamType Type = UE::AnimNext::UncookedOnly::FUtils::FindVariableType(VariableReference);				// The variable from the UAF Module
-		FAnimNextParamType PinType = UE::AnimNext::UncookedOnly::FUtils::GetParamTypeFromPinType(ValuePin->PinType);	// The variable from BP pin
-		UE::AnimNext::FParamCompatibility Compatibility = UE::AnimNext::FParamUtils::GetCompatibility(Type, PinType);
-		if (Compatibility == UE::AnimNext::EParamCompatibility::Incompatible || Compatibility == UE::AnimNext::EParamCompatibility::Incompatible_DataLoss)
+		FAnimNextParamType Type = UE::UAF::UncookedOnly::FUtils::FindVariableType(VariableReference);				// The variable from the UAF Module
+		FAnimNextParamType PinType = UE::UAF::UncookedOnly::FUtils::GetParamTypeFromPinType(ValuePin->PinType);	// The variable from BP pin
+		UE::UAF::FParamCompatibility Compatibility = UE::UAF::FParamUtils::GetCompatibility(Type, PinType);
+		if (Compatibility == UE::UAF::EParamCompatibility::Incompatible || Compatibility == UE::UAF::EParamCompatibility::Incompatible_DataLoss)
 		{
 			MessageLog.Error(*FString::Printf(TEXT("@@ mismatched type, setting '%s' type, but variable is '%s' type"), *PinType.ToString(), *Type.ToString()), this);
 		}
diff --git a/Engine/Plugins/Experimental/AnimNext/Source/AnimNextUncookedOnly/Private/Variables/UniversalObjectLocatorBindingType.h b/Engine/Plugins/Experimental/AnimNext/Source/AnimNextUncookedOnly/Private/Variables/UniversalObjectLocatorBindingType.h
--- a/Engine/Plugins/Experimental/AnimNext/Source/AnimNextUncookedOnly/Private/Variables/UniversalObjectLocatorBindingType.h
+++ b/Engine/Plugins/Experimental/AnimNext/Source/AnimNextUncookedOnly/Private/Variables/UniversalObjectLocatorBindingType.h
@@ -5,12 +5,12 @@
 #include "CoreMinimal.h"
 #include "Variables/IVariableBindingType.h"
 
-namespace UE::AnimNext::UncookedOnly
+namespace UE::UAF::UncookedOnly
 {
 	struct FClassProxy;
 }
 
-namespace UE::AnimNext::UncookedOnly
+namespace UE::UAF::UncookedOnly
 {
 
 // Provides information about object proxy parameter sources
diff --git a/Engine/Plugins/Experimental/AnimNext/Source/AnimNextUncookedOnly/Public/Variables/IVariableBindingType.h b/Engine/Plugins/Experimental/AnimNext/Source/AnimNextUncookedOnly/Public/Variables/IVariableBindingType.h
--- a/Engine/Plugins/Experimental/AnimNext/Source/AnimNextUncookedOnly/Public/Variables/IVariableBindingType.h
+++ b/Engine/Plugins/Experimental/AnimNext/Source/AnimNextUncookedOnly/Public/Variables/IVariableBindingType.h
@@ -14,18 +14,18 @@ class URigVMPin;
 struct FAnimNextVariableBindingData;
 template<typename T> struct TInstancedStruct;
 
-namespace UE::AnimNext::UncookedOnly
+namespace UE::UAF::UncookedOnly
 {
 	class FModule;
 	struct FUtils;
 }
 
-namespace UE::AnimNext::Editor
+namespace UE::UAF::Editor
 {
 	class FVariableBindingPropertyCustomization;
 }
 
-namespace UE::AnimNext::UncookedOnly
+namespace UE::UAF::UncookedOnly
 {
 
 // Interface used in editor/uncooked situations to determine the characteristics of a variable binding
diff --git a/Engine/Plugins/Experimental/AnimNextAnimGraph/Source/AnimNextAnimGraph/Internal/Graph/AnimNextGraphLatentPropertiesContextData.h b/Engine/Plugins/Experimental/AnimNextAnimGraph/Source/AnimNextAnimGraph/Internal/Graph/AnimNextGraphLatentPropertiesContextData.h
--- a/Engine/Plugins/Experimental/AnimNextAnimGraph/Source/AnimNextAnimGraph/Internal/Graph/AnimNextGraphLatentPropertiesContextData.h
+++ b/Engine/Plugins/Experimental/AnimNextAnimGraph/Source/AnimNextAnimGraph/Internal/Graph/AnimNextGraphLatentPropertiesContextData.h
@@ -8,7 +8,7 @@
 #include "Graph/AnimNextGraphInstance.h"
 #include "AnimNextGraphLatentPropertiesContextData.generated.h"
 
-namespace UE::AnimNext
+namespace UE::UAF
 {
 	struct FLatentPropertyHandle;
 }
@@ -20,7 +20,7 @@ struct FAnimNextGraphLatentPropertiesContextData : public FAnimNextGraphContextD
 
 	FAnimNextGraphLatentPropertiesContextData() = default;
 
-	FAnimNextGraphLatentPropertiesContextData(FAnimNextModuleInstance* InModuleInstance, const FAnimNextGraphInstance* InInstance, const TConstArrayView<UE::AnimNext::FLatentPropertyHandle>& InLatentHandles, void* InDestinationBasePtr, bool bInIsFrozen, bool bInJustBecameRelevant)
+	FAnimNextGraphLatentPropertiesContextData(FAnimNextModuleInstance* InModuleInstance, const FAnimNextGraphInstance* InInstance, const TConstArrayView<UE::UAF::FLatentPropertyHandle>& InLatentHandles, void* InDestinationBasePtr, bool bInIsFrozen, bool bInJustBecameRelevant)
 		: FAnimNextGraphContextData(InModuleInstance, InInstance)
 		, LatentHandles(InLatentHandles)
 		, DestinationBasePtr(InDestinationBasePtr)
@@ -29,13 +29,13 @@ struct FAnimNextGraphLatentPropertiesContextData : public FAnimNextGraphContextD
 	{
 	}
 
-	const TConstArrayView<UE::AnimNext::FLatentPropertyHandle>& GetLatentHandles() const { return LatentHandles; }
+	const TConstArrayView<UE::UAF::FLatentPropertyHandle>& GetLatentHandles() const { return LatentHandles; }
 	void* GetDestinationBasePtr() const { return DestinationBasePtr; }
 	bool IsFrozen() const { return bIsFrozen; }
 	bool JustBecameRelevant() const { return bJustBecameRelevant; }
 
 private:
-	TConstArrayView<UE::AnimNext::FLatentPropertyHandle> LatentHandles;
+	TConstArrayView<UE::UAF::FLatentPropertyHandle> LatentHandles;
 	void* DestinationBasePtr = nullptr;
 	bool bIsFrozen = false;
 	bool bJustBecameRelevant = false;
diff --git a/Engine/Plugins/Experimental/AnimNextAnimGraph/Source/AnimNextAnimGraph/Internal/Graph/AnimNextModuleAnimGraphComponent.h b/Engine/Plugins/Experimental/AnimNextAnimGraph/Source/AnimNextAnimGraph/Internal/Graph/AnimNextModuleAnimGraphComponent.h
--- a/Engine/Plugins/Experimental/AnimNextAnimGraph/Source/AnimNextAnimGraph/Internal/Graph/AnimNextModuleAnimGraphComponent.h
+++ b/Engine/Plugins/Experimental/AnimNextAnimGraph/Source/AnimNextAnimGraph/Internal/Graph/AnimNextModuleAnimGraphComponent.h
@@ -10,7 +10,7 @@
 struct FAnimNextGraphInstance;
 class UAnimNextAnimationGraph;
 
-namespace UE::AnimNext
+namespace UE::UAF
 {
 	struct FVariableOverrides;
 	struct FGraphAllocationParams;
@@ -24,7 +24,7 @@ struct FAnimNextModuleAnimGraphComponent : public FUAFModuleInstanceComponent
 
 	FAnimNextModuleAnimGraphComponent() = default;
 
-	TWeakPtr<FAnimNextGraphInstance> AllocateInstance(const UAnimNextAnimationGraph* InAnimationGraph, const TSharedPtr<const UE::AnimNext::FVariableOverridesCollection>& InOverrides);
+	TWeakPtr<FAnimNextGraphInstance> AllocateInstance(const UAnimNextAnimationGraph* InAnimationGraph, const TSharedPtr<const UE::UAF::FVariableOverridesCollection>& InOverrides);
 
 	void ReleaseInstance(TWeakPtr<FAnimNextGraphInstance> InInstance);
 
diff --git a/Engine/Plugins/Experimental/AnimNextAnimGraph/Source/AnimNextAnimGraph/Internal/Graph/RigDecorator_AnimNextCppTrait.h b/Engine/Plugins/Experimental/AnimNextAnimGraph/Source/AnimNextAnimGraph/Internal/Graph/RigDecorator_AnimNextCppTrait.h
--- a/Engine/Plugins/Experimental/AnimNextAnimGraph/Source/AnimNextAnimGraph/Internal/Graph/RigDecorator_AnimNextCppTrait.h
+++ b/Engine/Plugins/Experimental/AnimNextAnimGraph/Source/AnimNextAnimGraph/Internal/Graph/RigDecorator_AnimNextCppTrait.h
@@ -8,7 +8,7 @@
 
 #define UE_API ANIMNEXTANIMGRAPH_API
 
-namespace UE::AnimNext { struct FTrait; }
+namespace UE::UAF { struct FTrait; }
 
 /**
  * AnimNext RigDecorator for all C++ traits.
@@ -31,7 +31,7 @@ struct FRigDecorator_AnimNextCppDecorator : public FRigVMTrait
 		return DecoratorSharedDataStruct;
 	}
 
-	UE_API const UE::AnimNext::FTrait* GetTrait() const;
+	UE_API const UE::UAF::FTrait* GetTrait() const;
 #endif
 };
 
diff --git a/Engine/Plugins/Experimental/AnimNextAnimGraph/Source/AnimNextAnimGraph/Internal/Injection/InjectionCallbackProxy.h b/Engine/Plugins/Experimental/AnimNextAnimGraph/Source/AnimNextAnimGraph/Internal/Injection/InjectionCallbackProxy.h
--- a/Engine/Plugins/Experimental/AnimNextAnimGraph/Source/AnimNextAnimGraph/Internal/Injection/InjectionCallbackProxy.h
+++ b/Engine/Plugins/Experimental/AnimNextAnimGraph/Source/AnimNextAnimGraph/Internal/Injection/InjectionCallbackProxy.h
@@ -64,9 +64,9 @@ public:
 	ANIMNEXTANIMGRAPH_API virtual void Cancel() override;
 
 protected:
-	ANIMNEXTANIMGRAPH_API void OnInjectionCompleted(const UE::AnimNext::FInjectionRequest& Request);
-	ANIMNEXTANIMGRAPH_API void OnInjectionInterrupted(const UE::AnimNext::FInjectionRequest& Request);
-	ANIMNEXTANIMGRAPH_API void OnInjectionBlendingOut(const UE::AnimNext::FInjectionRequest& Request);
+	ANIMNEXTANIMGRAPH_API void OnInjectionCompleted(const UE::UAF::FInjectionRequest& Request);
+	ANIMNEXTANIMGRAPH_API void OnInjectionInterrupted(const UE::UAF::FInjectionRequest& Request);
+	ANIMNEXTANIMGRAPH_API void OnInjectionBlendingOut(const UE::UAF::FInjectionRequest& Request);
 
 	// Attempts to play an object with the specified payload. Returns whether it started or not.
 	ANIMNEXTANIMGRAPH_API bool Inject(
@@ -74,14 +74,14 @@ protected:
 		FAnimNextVariableReference Site,
 		UObject* Object,
 		FAnimNextFactoryParams&& FactoryParams,
-		const UE::AnimNext::FInjectionBlendSettings& BlendInSettings,
-		const UE::AnimNext::FInjectionBlendSettings& BlendOutSettings);
+		const UE::UAF::FInjectionBlendSettings& BlendInSettings,
+		const UE::UAF::FInjectionBlendSettings& BlendOutSettings);
 
 private:
 	void Reset();
 
 	DECLARE_FUNCTION(execSetVariable);
 
-	UE::AnimNext::FInjectionRequestPtr PlayingRequest;
+	UE::UAF::FInjectionRequestPtr PlayingRequest;
 	bool bWasInterrupted = false;
 };
diff --git a/Engine/Plugins/Experimental/AnimNextAnimGraph/Source/AnimNextAnimGraph/Internal/PlayAnim/PlayAnimCallbackProxy.h b/Engine/Plugins/Experimental/AnimNextAnimGraph/Source/AnimNextAnimGraph/Internal/PlayAnim/PlayAnimCallbackProxy.h
--- a/Engine/Plugins/Experimental/AnimNextAnimGraph/Source/AnimNextAnimGraph/Internal/PlayAnim/PlayAnimCallbackProxy.h
+++ b/Engine/Plugins/Experimental/AnimNextAnimGraph/Source/AnimNextAnimGraph/Internal/PlayAnim/PlayAnimCallbackProxy.h
@@ -57,9 +57,9 @@ public:
 	//~ End UObject Interface
 
 protected:
-	ANIMNEXTANIMGRAPH_API void OnPlayAnimCompleted(const UE::AnimNext::FInjectionRequest& Request);
-	ANIMNEXTANIMGRAPH_API void OnPlayAnimInterrupted(const UE::AnimNext::FInjectionRequest& Request);
-	ANIMNEXTANIMGRAPH_API void OnPlayAnimBlendingOut(const UE::AnimNext::FInjectionRequest& Request);
+	ANIMNEXTANIMGRAPH_API void OnPlayAnimCompleted(const UE::UAF::FInjectionRequest& Request);
+	ANIMNEXTANIMGRAPH_API void OnPlayAnimInterrupted(const UE::UAF::FInjectionRequest& Request);
+	ANIMNEXTANIMGRAPH_API void OnPlayAnimBlendingOut(const UE::UAF::FInjectionRequest& Request);
 
 	// Attempts to play an animation with the specified settings. Returns whether it started or not.
 	ANIMNEXTANIMGRAPH_API bool Play(
@@ -68,12 +68,12 @@ protected:
 		UAnimSequence* AnimSequence,
 		float PlayRate,
 		float StartPosition,
-		const UE::AnimNext::FInjectionBlendSettings& BlendInSettings,
-		const UE::AnimNext::FInjectionBlendSettings& BlendOutSettings);
+		const UE::UAF::FInjectionBlendSettings& BlendInSettings,
+		const UE::UAF::FInjectionBlendSettings& BlendOutSettings);
 
 private:
 	void Reset();
 
-	UE::AnimNext::FInjectionRequestPtr PlayingRequest;
+	UE::UAF::FInjectionRequestPtr PlayingRequest;
 	bool bWasInterrupted = false;
 };
diff --git a/Engine/Plugins/Experimental/AnimNextAnimGraph/Source/AnimNextAnimGraph/Internal/TraitCore/NodeTemplateRegistry.h b/Engine/Plugins/Experimental/AnimNextAnimGraph/Source/AnimNextAnimGraph/Internal/TraitCore/NodeTemplateRegistry.h
--- a/Engine/Plugins/Experimental/AnimNextAnimGraph/Source/AnimNextAnimGraph/Internal/TraitCore/NodeTemplateRegistry.h
+++ b/Engine/Plugins/Experimental/AnimNextAnimGraph/Source/AnimNextAnimGraph/Internal/TraitCore/NodeTemplateRegistry.h
@@ -8,12 +8,12 @@
 
 #define UE_API ANIMNEXTANIMGRAPH_API
 
-namespace UE::AnimNext::AnimGraph
+namespace UE::UAF::AnimGraph
 {
 	class FAnimNextAnimGraphModule;
 }
 
-namespace UE::AnimNext
+namespace UE::UAF
 {
 	struct FNodeTemplate;
 
@@ -69,7 +69,7 @@ namespace UE::AnimNext
 		TMap<uint32, FNodeTemplateRegistryHandle>	TemplateUIDToHandleMap;
 		mutable FRWLock								TemplateLock;
 
-		friend class UE::AnimNext::AnimGraph::FAnimNextAnimGraphModule;
+		friend class UE::UAF::AnimGraph::FAnimNextAnimGraphModule;
 		friend class FTraitWriter;
 		friend struct FScopedClearNodeTemplateRegistry;
 	};
diff --git a/Engine/Plugins/Experimental/AnimNextAnimGraph/Source/AnimNextAnimGraph/Internal/TraitCore/TraitInterfaceRegistry.h b/Engine/Plugins/Experimental/AnimNextAnimGraph/Source/AnimNextAnimGraph/Internal/TraitCore/TraitInterfaceRegistry.h
--- a/Engine/Plugins/Experimental/AnimNextAnimGraph/Source/AnimNextAnimGraph/Internal/TraitCore/TraitInterfaceRegistry.h
+++ b/Engine/Plugins/Experimental/AnimNextAnimGraph/Source/AnimNextAnimGraph/Internal/TraitCore/TraitInterfaceRegistry.h
@@ -8,12 +8,12 @@
 
 #define UE_API ANIMNEXTANIMGRAPH_API
 
-namespace UE::AnimNext::AnimGraph
+namespace UE::UAF::AnimGraph
 {
 	class FAnimNextAnimGraphModule;
 }
 
-namespace UE::AnimNext
+namespace UE::UAF
 {
 	/**
 	 * FTraitInterfaceRegistry
@@ -67,7 +67,7 @@ namespace UE::AnimNext
 
 		TMap<FTraitInterfaceUIDRaw, FRegistryEntry>	TraitInterfaceUIDToEntryMap;
 
-		friend class UE::AnimNext::AnimGraph::FAnimNextAnimGraphModule;
+		friend class UE::UAF::AnimGraph::FAnimNextAnimGraphModule;
 		friend struct FTraitInterfaceStaticInitHook;
 	};
 }
diff --git a/Engine/Plugins/Experimental/AnimNextAnimGraph/Source/AnimNextAnimGraph/Private/AnimNextAnimGraphModule.cpp b/Engine/Plugins/Experimental/AnimNextAnimGraph/Source/AnimNextAnimGraph/Private/AnimNextAnimGraphModule.cpp
--- a/Engine/Plugins/Experimental/AnimNextAnimGraph/Source/AnimNextAnimGraph/Private/AnimNextAnimGraphModule.cpp
+++ b/Engine/Plugins/Experimental/AnimNextAnimGraph/Source/AnimNextAnimGraph/Private/AnimNextAnimGraphModule.cpp
@@ -37,7 +37,7 @@
 #include "TraitCore/TraitTemplate.h"
 #endif
 
-namespace UE::AnimNext::AnimGraph
+namespace UE::UAF::AnimGraph
 {
 
 void FAnimNextAnimGraphModule::StartupModule()
@@ -342,4 +342,4 @@ void FAnimNextAnimGraphModule::ListAnimationGraphs(const TArray<FString>& Args)
 
 }
 
-IMPLEMENT_MODULE(UE::AnimNext::AnimGraph::FAnimNextAnimGraphModule, AnimNextAnimGraph)
+IMPLEMENT_MODULE(UE::UAF::AnimGraph::FAnimNextAnimGraphModule, AnimNextAnimGraph)
diff --git a/Engine/Plugins/Experimental/AnimNextAnimGraph/Source/AnimNextAnimGraph/Private/AnimNextAnimGraphModule.h b/Engine/Plugins/Experimental/AnimNextAnimGraph/Source/AnimNextAnimGraph/Private/AnimNextAnimGraphModule.h
--- a/Engine/Plugins/Experimental/AnimNextAnimGraph/Source/AnimNextAnimGraph/Private/AnimNextAnimGraphModule.h
+++ b/Engine/Plugins/Experimental/AnimNextAnimGraph/Source/AnimNextAnimGraph/Private/AnimNextAnimGraphModule.h
@@ -12,7 +12,7 @@
 class UAnimNextAnimationGraph;
 class UAnimNextAnimGraphSettings;
 
-namespace UE::AnimNext::AnimGraph
+namespace UE::UAF::AnimGraph
 {
 
 class FAnimNextAnimGraphModule : public IModuleInterface, public FGCObject
diff --git a/Engine/Plugins/Experimental/AnimNextAnimGraph/Source/AnimNextAnimGraph/Private/AnimNextAnimGraphSettings.cpp b/Engine/Plugins/Experimental/AnimNextAnimGraph/Source/AnimNextAnimGraph/Private/AnimNextAnimGraphSettings.cpp
--- a/Engine/Plugins/Experimental/AnimNextAnimGraph/Source/AnimNextAnimGraph/Private/AnimNextAnimGraphSettings.cpp
+++ b/Engine/Plugins/Experimental/AnimNextAnimGraph/Source/AnimNextAnimGraph/Private/AnimNextAnimGraphSettings.cpp
@@ -6,6 +6,6 @@
 TArray<UClass*> UAnimNextAnimGraphSettings::GetAllowedAssetClasses()
 {
 	TArray<UClass*> Array;
-	Array = UE::AnimNext::FAnimGraphFactory::GetRegisteredClasses();
+	Array = UE::UAF::FAnimGraphFactory::GetRegisteredClasses();
 	return Array;
 }
diff --git a/Engine/Plugins/Experimental/AnimNextAnimGraph/Source/AnimNextAnimGraph/Private/EvaluationVM/EvaluationVM.cpp b/Engine/Plugins/Experimental/AnimNextAnimGraph/Source/AnimNextAnimGraph/Private/EvaluationVM/EvaluationVM.cpp
--- a/Engine/Plugins/Experimental/AnimNextAnimGraph/Source/AnimNextAnimGraph/Private/EvaluationVM/EvaluationVM.cpp
+++ b/Engine/Plugins/Experimental/AnimNextAnimGraph/Source/AnimNextAnimGraph/Private/EvaluationVM/EvaluationVM.cpp
@@ -6,7 +6,7 @@
 #include "BoneContainer.h"
 #include "EvaluationVM/KeyframeState.h"
 
-namespace UE::AnimNext
+namespace UE::UAF
 {
 	const FEvaluationVMStackName KEYFRAME_STACK_NAME = FName(TEXT("KeyframeStack"));
 
@@ -28,7 +28,7 @@ namespace UE::AnimNext
 		}
 	}
 
-	FEvaluationVM::FEvaluationVM(EEvaluationFlags InEvaluationFlags, const UE::AnimNext::FReferencePose& InReferencePose, int32 InCurrentLOD)
+	FEvaluationVM::FEvaluationVM(EEvaluationFlags InEvaluationFlags, const UE::UAF::FReferencePose& InReferencePose, int32 InCurrentLOD)
 		: ReferencePose(&InReferencePose)
 		, CurrentLOD(InCurrentLOD)
 		, EvaluationFlags(InEvaluationFlags)
diff --git a/Engine/Plugins/Experimental/AnimNextAnimGraph/Source/AnimNextAnimGraph/Private/EvaluationVM/SerializableEvaluationProgram.cpp b/Engine/Plugins/Experimental/AnimNextAnimGraph/Source/AnimNextAnimGraph/Private/EvaluationVM/SerializableEvaluationProgram.cpp
--- a/Engine/Plugins/Experimental/AnimNextAnimGraph/Source/AnimNextAnimGraph/Private/EvaluationVM/SerializableEvaluationProgram.cpp
+++ b/Engine/Plugins/Experimental/AnimNextAnimGraph/Source/AnimNextAnimGraph/Private/EvaluationVM/SerializableEvaluationProgram.cpp
@@ -3,7 +3,7 @@
 #include "EvaluationVM/SerializableEvaluationProgram.h"
 #include "EvaluationVM/EvaluationProgram.h"
 
-FSerializableEvaluationProgram::FSerializableEvaluationProgram(const UE::AnimNext::FEvaluationProgram& Other)
+FSerializableEvaluationProgram::FSerializableEvaluationProgram(const UE::UAF::FEvaluationProgram& Other)
 {
 	for (const TSharedPtr<FAnimNextEvaluationTask>& Task : Other.Tasks)
 	{
diff --git a/Engine/Plugins/Experimental/AnimNextAnimGraph/Source/AnimNextAnimGraph/Private/EvaluationVM/Tasks/ApplyAdditiveKeyframe.cpp b/Engine/Plugins/Experimental/AnimNextAnimGraph/Source/AnimNextAnimGraph/Private/EvaluationVM/Tasks/ApplyAdditiveKeyframe.cpp
--- a/Engine/Plugins/Experimental/AnimNextAnimGraph/Source/AnimNextAnimGraph/Private/EvaluationVM/Tasks/ApplyAdditiveKeyframe.cpp
+++ b/Engine/Plugins/Experimental/AnimNextAnimGraph/Source/AnimNextAnimGraph/Private/EvaluationVM/Tasks/ApplyAdditiveKeyframe.cpp
@@ -23,9 +23,9 @@ FAnimNextApplyAdditiveKeyframeTask FAnimNextApplyAdditiveKeyframeTask::Make(cons
 	return Task;
 }
 
-void FAnimNextApplyAdditiveKeyframeTask::Execute(UE::AnimNext::FEvaluationVM& VM) const
+void FAnimNextApplyAdditiveKeyframeTask::Execute(UE::UAF::FEvaluationVM& VM) const
 {
-	using namespace UE::AnimNext;
+	using namespace UE::UAF;
 
 	// Pop our top two poses, we'll re-use the top keyframe for our result
 
@@ -87,7 +87,7 @@ void FAnimNextApplyAdditiveKeyframeTask::Execute(UE::AnimNext::FEvaluationVM& VM
 	VM.PushValue(KEYFRAME_STACK_NAME, MoveTemp(BaseKeyframe));
 }
 
-float FAnimNextApplyAdditiveKeyframeTask::GetInterpolationAlpha(const UE::AnimNext::FKeyframeState* KeyframeA, const UE::AnimNext::FKeyframeState* KeyframeB) const
+float FAnimNextApplyAdditiveKeyframeTask::GetInterpolationAlpha(const UE::UAF::FKeyframeState* KeyframeA, const UE::UAF::FKeyframeState* KeyframeB) const
 {
 	float Alpha = BlendWeight;
 
diff --git a/Engine/Plugins/Experimental/AnimNextAnimGraph/Source/AnimNextAnimGraph/Private/EvaluationVM/Tasks/BlendKeyframes.cpp b/Engine/Plugins/Experimental/AnimNextAnimGraph/Source/AnimNextAnimGraph/Private/EvaluationVM/Tasks/BlendKeyframes.cpp
--- a/Engine/Plugins/Experimental/AnimNextAnimGraph/Source/AnimNextAnimGraph/Private/EvaluationVM/Tasks/BlendKeyframes.cpp
+++ b/Engine/Plugins/Experimental/AnimNextAnimGraph/Source/AnimNextAnimGraph/Private/EvaluationVM/Tasks/BlendKeyframes.cpp
@@ -24,9 +24,9 @@ FAnimNextBlendTwoKeyframesTask FAnimNextBlendTwoKeyframesTask::Make(const FName&
 	return Task;
 }
 
-void FAnimNextBlendTwoKeyframesTask::Execute(UE::AnimNext::FEvaluationVM& VM) const
+void FAnimNextBlendTwoKeyframesTask::Execute(UE::UAF::FEvaluationVM& VM) const
 {
-	using namespace UE::AnimNext;
+	using namespace UE::UAF;
 
 	// Pop our top two poses, we'll re-use the top keyframe for our result
 
@@ -92,7 +92,7 @@ void FAnimNextBlendTwoKeyframesTask::Execute(UE::AnimNext::FEvaluationVM& VM) co
 	}
 }
 
-float FAnimNextBlendTwoKeyframesTask::GetInterpolationAlpha(const UE::AnimNext::FKeyframeState* KeyframeA, const UE::AnimNext::FKeyframeState* KeyframeB) const
+float FAnimNextBlendTwoKeyframesTask::GetInterpolationAlpha(const UE::UAF::FKeyframeState* KeyframeA, const UE::UAF::FKeyframeState* KeyframeB) const
 {
 	float Alpha = InterpolationAlpha;
 
@@ -123,9 +123,9 @@ FAnimNextBlendOverwriteKeyframeWithScaleTask FAnimNextBlendOverwriteKeyframeWith
 	return Task;
 }
 
-void FAnimNextBlendOverwriteKeyframeWithScaleTask::Execute(UE::AnimNext::FEvaluationVM& VM) const
+void FAnimNextBlendOverwriteKeyframeWithScaleTask::Execute(UE::UAF::FEvaluationVM& VM) const
 {
-	using namespace UE::AnimNext;
+	using namespace UE::UAF;
 
 	TUniquePtr<FKeyframeState> Keyframe;
 	if (!VM.PopValue(KEYFRAME_STACK_NAME, Keyframe))
@@ -166,9 +166,9 @@ FAnimNextBlendAddKeyframeWithScaleTask FAnimNextBlendAddKeyframeWithScaleTask::M
 	return Task;
 }
 
-void FAnimNextBlendAddKeyframeWithScaleTask::Execute(UE::AnimNext::FEvaluationVM& VM) const
+void FAnimNextBlendAddKeyframeWithScaleTask::Execute(UE::UAF::FEvaluationVM& VM) const
 {
-	using namespace UE::AnimNext;
+	using namespace UE::UAF;
 
 	// Pop our top two poses, we'll re-use the top keyframe for our result
 
diff --git a/Engine/Plugins/Experimental/AnimNextAnimGraph/Source/AnimNextAnimGraph/Private/EvaluationVM/Tasks/BlendKeyframesPerBone.cpp b/Engine/Plugins/Experimental/AnimNextAnimGraph/Source/AnimNextAnimGraph/Private/EvaluationVM/Tasks/BlendKeyframesPerBone.cpp
--- a/Engine/Plugins/Experimental/AnimNextAnimGraph/Source/AnimNextAnimGraph/Private/EvaluationVM/Tasks/BlendKeyframesPerBone.cpp
+++ b/Engine/Plugins/Experimental/AnimNextAnimGraph/Source/AnimNextAnimGraph/Private/EvaluationVM/Tasks/BlendKeyframesPerBone.cpp
@@ -23,9 +23,9 @@ FAnimNextBlendOverwriteKeyframePerBoneWithScaleTask FAnimNextBlendOverwriteKeyfr
 	return Task;
 }
 
-void FAnimNextBlendOverwriteKeyframePerBoneWithScaleTask::Execute(UE::AnimNext::FEvaluationVM& VM) const
+void FAnimNextBlendOverwriteKeyframePerBoneWithScaleTask::Execute(UE::UAF::FEvaluationVM& VM) const
 {
-	using namespace UE::AnimNext;
+	using namespace UE::UAF;
 
 	if (BlendProfile == nullptr)
 	{
@@ -109,9 +109,9 @@ FAnimNextBlendAddKeyframePerBoneWithScaleTask FAnimNextBlendAddKeyframePerBoneWi
 	return Task;
 }
 
-void FAnimNextBlendAddKeyframePerBoneWithScaleTask::Execute(UE::AnimNext::FEvaluationVM& VM) const
+void FAnimNextBlendAddKeyframePerBoneWithScaleTask::Execute(UE::UAF::FEvaluationVM& VM) const
 {
-	using namespace UE::AnimNext;
+	using namespace UE::UAF;
 
 	if (BlendProfile == nullptr)
 	{
@@ -212,9 +212,9 @@ FAnimNextBlendKeyframePerBoneWithScaleTask FAnimNextBlendKeyframePerBoneWithScal
 	return Task;
 }
 
-void FAnimNextBlendKeyframePerBoneWithScaleTask::Execute(UE::AnimNext::FEvaluationVM& VM) const
+void FAnimNextBlendKeyframePerBoneWithScaleTask::Execute(UE::UAF::FEvaluationVM& VM) const
 {
-	using namespace UE::AnimNext;
+	using namespace UE::UAF;
 
 	if (BlendProfile == nullptr)
 	{
diff --git a/Engine/Plugins/Experimental/AnimNextAnimGraph/Source/AnimNextAnimGraph/Private/EvaluationVM/Tasks/ConvertRotationsToLocalSpace.cpp b/Engine/Plugins/Experimental/AnimNextAnimGraph/Source/AnimNextAnimGraph/Private/EvaluationVM/Tasks/ConvertRotationsToLocalSpace.cpp
--- a/Engine/Plugins/Experimental/AnimNextAnimGraph/Source/AnimNextAnimGraph/Private/EvaluationVM/Tasks/ConvertRotationsToLocalSpace.cpp
+++ b/Engine/Plugins/Experimental/AnimNextAnimGraph/Source/AnimNextAnimGraph/Private/EvaluationVM/Tasks/ConvertRotationsToLocalSpace.cpp
@@ -17,9 +17,9 @@ FAnimNextConvertRotationsToLocalSpaceTask FAnimNextConvertRotationsToLocalSpaceT
 	return Task;
 }
 
-void FAnimNextConvertRotationsToLocalSpaceTask::Execute(UE::AnimNext::FEvaluationVM& VM) const
+void FAnimNextConvertRotationsToLocalSpaceTask::Execute(UE::UAF::FEvaluationVM& VM) const
 {
-	using namespace UE::AnimNext;
+	using namespace UE::UAF;
 
 	if (!EnumHasAnyFlags(VM.GetFlags(), EEvaluationFlags::Bones))
 	{
diff --git a/Engine/Plugins/Experimental/AnimNextAnimGraph/Source/AnimNextAnimGraph/Private/EvaluationVM/Tasks/ConvertRotationsToMeshSpace.cpp b/Engine/Plugins/Experimental/AnimNextAnimGraph/Source/AnimNextAnimGraph/Private/EvaluationVM/Tasks/ConvertRotationsToMeshSpace.cpp
--- a/Engine/Plugins/Experimental/AnimNextAnimGraph/Source/AnimNextAnimGraph/Private/EvaluationVM/Tasks/ConvertRotationsToMeshSpace.cpp
+++ b/Engine/Plugins/Experimental/AnimNextAnimGraph/Source/AnimNextAnimGraph/Private/EvaluationVM/Tasks/ConvertRotationsToMeshSpace.cpp
@@ -17,9 +17,9 @@ FAnimNextConvertRotationsToMeshSpaceTask FAnimNextConvertRotationsToMeshSpaceTas
 	return Task;
 }
 
-void FAnimNextConvertRotationsToMeshSpaceTask::Execute(UE::AnimNext::FEvaluationVM& VM) const
+void FAnimNextConvertRotationsToMeshSpaceTask::Execute(UE::UAF::FEvaluationVM& VM) const
 {
-	using namespace UE::AnimNext;
+	using namespace UE::UAF;
 
 	if (!EnumHasAnyFlags(VM.GetFlags(), EEvaluationFlags::Bones))
 	{
diff --git a/Engine/Plugins/Experimental/AnimNextAnimGraph/Source/AnimNextAnimGraph/Private/EvaluationVM/Tasks/DeadBlending.cpp b/Engine/Plugins/Experimental/AnimNextAnimGraph/Source/AnimNextAnimGraph/Private/EvaluationVM/Tasks/DeadBlending.cpp
--- a/Engine/Plugins/Experimental/AnimNextAnimGraph/Source/AnimNextAnimGraph/Private/EvaluationVM/Tasks/DeadBlending.cpp
+++ b/Engine/Plugins/Experimental/AnimNextAnimGraph/Source/AnimNextAnimGraph/Private/EvaluationVM/Tasks/DeadBlending.cpp
@@ -17,7 +17,7 @@
 #include "DeadBlending.ispc.generated.h"
 #endif
 
-namespace UE::AnimNext::DeadBlending::Private
+namespace UE::UAF::DeadBlending::Private
 {
 #if !UE_ANIM_NEXT_DEAD_BLENDING_ISPC
 	static constexpr float Ln2 = 0.69314718056f;
@@ -197,9 +197,9 @@ namespace UE::AnimNext::DeadBlending::Private
 		const FTransformArraySoAConstView& Curr,
 		const FTransformArraySoAConstView& Prev,
 		const float DeltaTime,
-		const UE::AnimNext::FDeadBlendTransitionTaskParameters& Parameters)
+		const UE::UAF::FDeadBlendTransitionTaskParameters& Parameters)
 	{
-		TRACE_CPUPROFILER_EVENT_SCOPE(AnimNext::DeadBlending::Transition);
+		TRACE_CPUPROFILER_EVENT_SCOPE(UAF::DeadBlending::Transition);
 
 		const int32 LODBoneNum = Source.Num();
 		check(LODBoneNum <= Dest.Num());
@@ -311,9 +311,9 @@ namespace UE::AnimNext::DeadBlending::Private
 		const TArrayView<FVector3f>& SourceBoneRotationDecayHalfLives,
 		const TArrayView<FVector3f>& SourceBoneScaleDecayHalfLives,
 		const FTransformArraySoAConstView& Curr,
-		const UE::AnimNext::FDeadBlendTransitionTaskParameters& Parameters)
+		const UE::UAF::FDeadBlendTransitionTaskParameters& Parameters)
 	{
-		TRACE_CPUPROFILER_EVENT_SCOPE(AnimNext::DeadBlending::Transition);
+		TRACE_CPUPROFILER_EVENT_SCOPE(UAF::DeadBlending::Transition);
 
 		const int32 LODBoneNum = Source.Num();
 		check(LODBoneNum <= Curr.Num());
@@ -370,7 +370,7 @@ namespace UE::AnimNext::DeadBlending::Private
 		const EAlphaBlendOption BlendMode,
 		UCurveFloat* CustomBlendCurve)
 	{
-		TRACE_CPUPROFILER_EVENT_SCOPE(AnimNext::DeadBlending::Apply)
+		TRACE_CPUPROFILER_EVENT_SCOPE(UAF::DeadBlending::Apply)
 
 		const float Alpha = 1.0f - FAlphaBlend::AlphaToBlendOption(
 			FMath::Clamp(TimeSinceTransition / FMath::Max(BlendDuration, UE_SMALL_NUMBER), 0.0f, 1.0f),
@@ -451,11 +451,11 @@ namespace UE::AnimNext::DeadBlending::Private
 }
 
 FAnimNextDeadBlendingTransitionTask FAnimNextDeadBlendingTransitionTask::Make(
-	UE::AnimNext::FDeadBlendingState* State,
-	const UE::AnimNext::FTransformArraySoAHeap* CurrPose,
-	const UE::AnimNext::FTransformArraySoAHeap* PrevPose,
+	UE::UAF::FDeadBlendingState* State,
+	const UE::UAF::FTransformArraySoAHeap* CurrPose,
+	const UE::UAF::FTransformArraySoAHeap* PrevPose,
 	const float DeltaTime,
-	const UE::AnimNext::FDeadBlendTransitionTaskParameters& Parameters)
+	const UE::UAF::FDeadBlendTransitionTaskParameters& Parameters)
 {
 	FAnimNextDeadBlendingTransitionTask Task;
 	Task.State = State;
@@ -467,9 +467,9 @@ FAnimNextDeadBlendingTransitionTask FAnimNextDeadBlendingTransitionTask::Make(
 }
 
 FAnimNextDeadBlendingTransitionTask FAnimNextDeadBlendingTransitionTask::Make(
-	UE::AnimNext::FDeadBlendingState* State,
-	const UE::AnimNext::FTransformArraySoAHeap* CurrPose,
-	const UE::AnimNext::FDeadBlendTransitionTaskParameters& Parameters)
+	UE::UAF::FDeadBlendingState* State,
+	const UE::UAF::FTransformArraySoAHeap* CurrPose,
+	const UE::UAF::FDeadBlendTransitionTaskParameters& Parameters)
 {
 	FAnimNextDeadBlendingTransitionTask Task;
 	Task.State = State;
@@ -478,9 +478,9 @@ FAnimNextDeadBlendingTransitionTask FAnimNextDeadBlendingTransitionTask::Make(
 	return Task;
 }
 
-void FAnimNextDeadBlendingTransitionTask::Execute(UE::AnimNext::FEvaluationVM& VM) const
+void FAnimNextDeadBlendingTransitionTask::Execute(UE::UAF::FEvaluationVM& VM) const
 {
-	using namespace UE::AnimNext;
+	using namespace UE::UAF;
 
 	if (EnumHasAnyFlags(VM.GetFlags(), EEvaluationFlags::Bones))
 	{
@@ -532,7 +532,7 @@ void FAnimNextDeadBlendingTransitionTask::Execute(UE::AnimNext::FEvaluationVM& V
 }
 
 FAnimNextDeadBlendingApplyTask FAnimNextDeadBlendingApplyTask::Make(
-	UE::AnimNext::FDeadBlendingState* State,
+	UE::UAF::FDeadBlendingState* State,
 	const float BlendDuration,
 	const float TimeSinceTransition,
 	const EAlphaBlendOption& BlendMode,
@@ -547,9 +547,9 @@ FAnimNextDeadBlendingApplyTask FAnimNextDeadBlendingApplyTask::Make(
 	return Task;
 }
 
-void FAnimNextDeadBlendingApplyTask::Execute(UE::AnimNext::FEvaluationVM& VM) const
+void FAnimNextDeadBlendingApplyTask::Execute(UE::UAF::FEvaluationVM& VM) const
 {
-	using namespace UE::AnimNext;
+	using namespace UE::UAF;
 
 	if (EnumHasAnyFlags(VM.GetFlags(), EEvaluationFlags::Bones))
 	{
diff --git a/Engine/Plugins/Experimental/AnimNextAnimGraph/Source/AnimNextAnimGraph/Private/EvaluationVM/Tasks/ExecuteProgram.cpp b/Engine/Plugins/Experimental/AnimNextAnimGraph/Source/AnimNextAnimGraph/Private/EvaluationVM/Tasks/ExecuteProgram.cpp
--- a/Engine/Plugins/Experimental/AnimNextAnimGraph/Source/AnimNextAnimGraph/Private/EvaluationVM/Tasks/ExecuteProgram.cpp
+++ b/Engine/Plugins/Experimental/AnimNextAnimGraph/Source/AnimNextAnimGraph/Private/EvaluationVM/Tasks/ExecuteProgram.cpp
@@ -5,14 +5,14 @@
 #include "EvaluationVM/EvaluationProgram.h"
 #include "EvaluationVM/EvaluationVM.h"
 
-FAnimNextExecuteProgramTask FAnimNextExecuteProgramTask::Make(TSharedPtr<const UE::AnimNext::FEvaluationProgram> Program)
+FAnimNextExecuteProgramTask FAnimNextExecuteProgramTask::Make(TSharedPtr<const UE::UAF::FEvaluationProgram> Program)
 {
 	FAnimNextExecuteProgramTask Task;
 	Task.Program = Program;
 	return Task;
 }
 
-void FAnimNextExecuteProgramTask::Execute(UE::AnimNext::FEvaluationVM& VM) const
+void FAnimNextExecuteProgramTask::Execute(UE::UAF::FEvaluationVM& VM) const
 {
 	if (Program.IsValid())
 	{
diff --git a/Engine/Plugins/Experimental/AnimNextAnimGraph/Source/AnimNextAnimGraph/Private/EvaluationVM/Tasks/NormalizeRotations.cpp b/Engine/Plugins/Experimental/AnimNextAnimGraph/Source/AnimNextAnimGraph/Private/EvaluationVM/Tasks/NormalizeRotations.cpp
--- a/Engine/Plugins/Experimental/AnimNextAnimGraph/Source/AnimNextAnimGraph/Private/EvaluationVM/Tasks/NormalizeRotations.cpp
+++ b/Engine/Plugins/Experimental/AnimNextAnimGraph/Source/AnimNextAnimGraph/Private/EvaluationVM/Tasks/NormalizeRotations.cpp
@@ -6,9 +6,9 @@
 #include "EvaluationVM/KeyframeState.h"
 #include "TransformArrayOperations.h"
 
-void FAnimNextNormalizeKeyframeRotationsTask::Execute(UE::AnimNext::FEvaluationVM& VM) const
+void FAnimNextNormalizeKeyframeRotationsTask::Execute(UE::UAF::FEvaluationVM& VM) const
 {
-	using namespace UE::AnimNext;
+	using namespace UE::UAF;
 
 	if (EnumHasAnyFlags(VM.GetFlags(), EEvaluationFlags::Bones))
 	{
diff --git a/Engine/Plugins/Experimental/AnimNextAnimGraph/Source/AnimNextAnimGraph/Private/EvaluationVM/Tasks/PushAnimSequenceKeyframe.cpp b/Engine/Plugins/Experimental/AnimNextAnimGraph/Source/AnimNextAnimGraph/Private/EvaluationVM/Tasks/PushAnimSequenceKeyframe.cpp
--- a/Engine/Plugins/Experimental/AnimNextAnimGraph/Source/AnimNextAnimGraph/Private/EvaluationVM/Tasks/PushAnimSequenceKeyframe.cpp
+++ b/Engine/Plugins/Experimental/AnimNextAnimGraph/Source/AnimNextAnimGraph/Private/EvaluationVM/Tasks/PushAnimSequenceKeyframe.cpp
@@ -28,9 +28,9 @@ FAnimNextAnimSequenceKeyframeTask FAnimNextAnimSequenceKeyframeTask::MakeFromKey
 	return Task;
 }
 
-void FAnimNextAnimSequenceKeyframeTask::Execute(UE::AnimNext::FEvaluationVM& VM) const
+void FAnimNextAnimSequenceKeyframeTask::Execute(UE::UAF::FEvaluationVM& VM) const
 {
-	using namespace UE::AnimNext;
+	using namespace UE::UAF;
 	
 	if(const UAnimSequence* AnimSequencePtr = AnimSequence.Get())
 	{
diff --git a/Engine/Plugins/Experimental/AnimNextAnimGraph/Source/AnimNextAnimGraph/Private/EvaluationVM/Tasks/PushPose.cpp b/Engine/Plugins/Experimental/AnimNextAnimGraph/Source/AnimNextAnimGraph/Private/EvaluationVM/Tasks/PushPose.cpp
--- a/Engine/Plugins/Experimental/AnimNextAnimGraph/Source/AnimNextAnimGraph/Private/EvaluationVM/Tasks/PushPose.cpp
+++ b/Engine/Plugins/Experimental/AnimNextAnimGraph/Source/AnimNextAnimGraph/Private/EvaluationVM/Tasks/PushPose.cpp
@@ -21,11 +21,11 @@ FAnimNextPushPoseTask& FAnimNextPushPoseTask::operator=(const FAnimNextPushPoseT
 	return *this;
 }
 
-void FAnimNextPushPoseTask::Execute(UE::AnimNext::FEvaluationVM& VM) const
+void FAnimNextPushPoseTask::Execute(UE::UAF::FEvaluationVM& VM) const
 {
 	SCOPE_CYCLE_COUNTER(STAT_AnimNext_Task_PushPose);
 
-	using namespace UE::AnimNext;
+	using namespace UE::UAF;
 
 	FKeyframeState Keyframe = VM.MakeUninitializedKeyframe(GraphPose->LODPose.IsAdditive());
 
diff --git a/Engine/Plugins/Experimental/AnimNextAnimGraph/Source/AnimNextAnimGraph/Private/EvaluationVM/Tasks/PushReferenceKeyframe.cpp b/Engine/Plugins/Experimental/AnimNextAnimGraph/Source/AnimNextAnimGraph/Private/EvaluationVM/Tasks/PushReferenceKeyframe.cpp
--- a/Engine/Plugins/Experimental/AnimNextAnimGraph/Source/AnimNextAnimGraph/Private/EvaluationVM/Tasks/PushReferenceKeyframe.cpp
+++ b/Engine/Plugins/Experimental/AnimNextAnimGraph/Source/AnimNextAnimGraph/Private/EvaluationVM/Tasks/PushReferenceKeyframe.cpp
@@ -17,9 +17,9 @@ FAnimNextPushReferenceKeyframeTask FAnimNextPushReferenceKeyframeTask::MakeFromA
 	return Task;
 }
 
-void FAnimNextPushReferenceKeyframeTask::Execute(UE::AnimNext::FEvaluationVM& VM) const
+void FAnimNextPushReferenceKeyframeTask::Execute(UE::UAF::FEvaluationVM& VM) const
 {
-	using namespace UE::AnimNext;
+	using namespace UE::UAF;
 
 	VM.PushValue(KEYFRAME_STACK_NAME, MakeUnique<FKeyframeState>(VM.MakeReferenceKeyframe(bIsAdditive)));
 }
diff --git a/Engine/Plugins/Experimental/AnimNextAnimGraph/Source/AnimNextAnimGraph/Private/EvaluationVM/Tasks/StoreKeyframe.cpp b/Engine/Plugins/Experimental/AnimNextAnimGraph/Source/AnimNextAnimGraph/Private/EvaluationVM/Tasks/StoreKeyframe.cpp
--- a/Engine/Plugins/Experimental/AnimNextAnimGraph/Source/AnimNextAnimGraph/Private/EvaluationVM/Tasks/StoreKeyframe.cpp
+++ b/Engine/Plugins/Experimental/AnimNextAnimGraph/Source/AnimNextAnimGraph/Private/EvaluationVM/Tasks/StoreKeyframe.cpp
@@ -6,7 +6,7 @@
 #include "EvaluationVM/KeyframeState.h"
 #include "TransformArrayOperations.h"
 
-FAnimNextSwapTransformsTask FAnimNextSwapTransformsTask::Make(UE::AnimNext::FTransformArraySoAHeap* A, UE::AnimNext::FTransformArraySoAHeap* B)
+FAnimNextSwapTransformsTask FAnimNextSwapTransformsTask::Make(UE::UAF::FTransformArraySoAHeap* A, UE::UAF::FTransformArraySoAHeap* B)
 {
 	FAnimNextSwapTransformsTask Task;
 	Task.A = A;
@@ -14,21 +14,21 @@ FAnimNextSwapTransformsTask FAnimNextSwapTransformsTask::Make(UE::AnimNext::FTra
 	return Task;
 }
 
-void FAnimNextSwapTransformsTask::Execute(UE::AnimNext::FEvaluationVM& VM) const
+void FAnimNextSwapTransformsTask::Execute(UE::UAF::FEvaluationVM& VM) const
 {
 	Swap(*A, *B);
 }
 
-FAnimNextStoreKeyframeTransformsTask FAnimNextStoreKeyframeTransformsTask::Make(UE::AnimNext::FTransformArraySoAHeap* Dest)
+FAnimNextStoreKeyframeTransformsTask FAnimNextStoreKeyframeTransformsTask::Make(UE::UAF::FTransformArraySoAHeap* Dest)
 {
 	FAnimNextStoreKeyframeTransformsTask Task;
 	Task.Dest = Dest;
 	return Task;
 }
 
-void FAnimNextStoreKeyframeTransformsTask::Execute(UE::AnimNext::FEvaluationVM& VM) const
+void FAnimNextStoreKeyframeTransformsTask::Execute(UE::UAF::FEvaluationVM& VM) const
 {
-	using namespace UE::AnimNext;
+	using namespace UE::UAF;
 
 	if (EnumHasAnyFlags(VM.GetFlags(), EEvaluationFlags::Bones))
 	{
@@ -49,9 +49,9 @@ FAnimNextDuplicateTopKeyframeTask FAnimNextDuplicateTopKeyframeTask::Make()
 }
 
 // Task entry point
-void FAnimNextDuplicateTopKeyframeTask::Execute(UE::AnimNext::FEvaluationVM& VM) const
+void FAnimNextDuplicateTopKeyframeTask::Execute(UE::UAF::FEvaluationVM& VM) const
 {
-	using namespace UE::AnimNext;
+	using namespace UE::UAF;
 
 	TUniquePtr<FKeyframeState> KeyframeOut;
 
