UAF Wizard 0.1

#virtualized

diff --git a/Engine/Plugins/Experimental/AnimNext/AnimNext.uplugin b/Engine/Plugins/Experimental/AnimNext/AnimNext.uplugin
--- a/Engine/Plugins/Experimental/AnimNext/AnimNext.uplugin
+++ b/Engine/Plugins/Experimental/AnimNext/AnimNext.uplugin
@@ -11,7 +11,7 @@
 	"MarketplaceURL" : "",
 	"SupportURL" : "",
 	"EnabledByDefault" : false,
-	"CanContainContent" : false,
+	"CanContainContent" : true,
 	"IsBetaVersion": false,
 	"IsExperimentalVersion": true,
 	"Installed" : false,
diff --git a/Engine/Plugins/Experimental/AnimNext/Source/AnimNext/Private/Component/AnimNextComponent.cpp b/Engine/Plugins/Experimental/AnimNext/Source/AnimNext/Private/Component/AnimNextComponent.cpp
--- a/Engine/Plugins/Experimental/AnimNext/Source/AnimNext/Private/Component/AnimNextComponent.cpp
+++ b/Engine/Plugins/Experimental/AnimNext/Source/AnimNext/Private/Component/AnimNextComponent.cpp
@@ -461,4 +461,9 @@ void UAnimNextComponent::RemoveModuleEventSubsequent(FName InEventName, UAnimNex
 FAnimNextModuleHandle UAnimNextComponent::BlueprintGetModuleHandle() const
 {
 	return FAnimNextModuleHandle(ModuleHandle);
-}
\ No newline at end of file
+}
+
+void UAnimNextComponent::SetModule(TObjectPtr<UAnimNextModule> InModule)
+{
+	Module = InModule;
+}
diff --git a/Engine/Plugins/Experimental/AnimNext/Source/AnimNext/Public/Component/AnimNextComponent.h b/Engine/Plugins/Experimental/AnimNext/Source/AnimNext/Public/Component/AnimNextComponent.h
--- a/Engine/Plugins/Experimental/AnimNext/Source/AnimNext/Public/Component/AnimNextComponent.h
+++ b/Engine/Plugins/Experimental/AnimNext/Source/AnimNext/Public/Component/AnimNextComponent.h
@@ -22,6 +22,11 @@ namespace UE::UAF::UncookedOnly
 	struct FUtils;
 }
 
+namespace UE::UAF::Editor
+{
+	class FAssetWizard;
+}
+
 UCLASS(MinimalAPI, meta = (BlueprintSpawnableComponent))
 class UAnimNextComponent : public UActorComponent
 {
@@ -180,6 +185,8 @@ public:
 	ANIMNEXT_API FAnimNextModuleHandle BlueprintGetModuleHandle() const;
 
 private:
+	ANIMNEXT_API void SetModule(TObjectPtr<UAnimNextModule> InModule);
+	
 	DECLARE_FUNCTION(execBlueprintSetVariable);
 	DECLARE_FUNCTION(execBlueprintSetVariableReference);
 
@@ -192,6 +199,7 @@ private:
 private:
 	friend struct UE::UAF::UncookedOnly::FUtils;
 	friend class UAnimNextComponentWorldSubsystem;
+	friend class UE::UAF::Editor::FAssetWizard;
 
 	// The AnimNext module that this component will run
 	UPROPERTY(EditAnywhere, Category="Module")
diff --git a/Engine/Plugins/Experimental/AnimNext/Source/AnimNextEditor/Internal/AnimNextStyle.h b/Engine/Plugins/Experimental/AnimNext/Source/AnimNextEditor/Internal/AnimNextStyle.h
--- a/Engine/Plugins/Experimental/AnimNext/Source/AnimNextEditor/Internal/AnimNextStyle.h
+++ b/Engine/Plugins/Experimental/AnimNext/Source/AnimNextEditor/Internal/AnimNextStyle.h
@@ -17,4 +17,6 @@ protected:
 
 private:
 	FAnimNextStyle();
+
+	void SetAssetWizardStyles();
 };
diff --git a/Engine/Plugins/Experimental/AnimNext/Source/AnimNextEditor/Private/AnimNextEditorModule.cpp b/Engine/Plugins/Experimental/AnimNext/Source/AnimNextEditor/Private/AnimNextEditorModule.cpp
--- a/Engine/Plugins/Experimental/AnimNext/Source/AnimNextEditor/Private/AnimNextEditorModule.cpp
+++ b/Engine/Plugins/Experimental/AnimNext/Source/AnimNextEditor/Private/AnimNextEditorModule.cpp
@@ -34,7 +34,9 @@
 #include "Param/ObjectPropertyLocatorEditor.h"
 #include "AnimNextAssetWorkspaceAssetUserData.h"
 #include "AnimNextEditorContext.h"
+#include "AnimNextStyle.h"
 #include "AssetCompilationHandler.h"
+#include "ContentBrowserModule.h"
 #include "ContextObjectStore.h"
 #include "EditorModeManager.h"
 #include "Common/AnimNextGraphItemDetails.h"
@@ -75,6 +77,7 @@
 #include "Common/AnimNextAssetFindReplaceVariables.h"
 #include "Toolkits/AssetEditorToolkitMenuContext.h"
 #include "Variables/Outliner/VariablesOutlinerCommands.h"
+#include "Wizard/AssetWizard.h"
 
 #define LOCTEXT_NAMESPACE "AnimNextEditorModule"
 
@@ -91,7 +94,7 @@ namespace
 
 namespace UE::UAF::Editor
 {
-
+	
 void FAnimNextEditorModule::StartupModule()
 {
 	FVariableOverrideCommands::Register();
@@ -107,6 +110,18 @@ void FAnimNextEditorModule::StartupModule()
 		LOCTEXT("SettingsDescription", "Customize AnimNext Settings."),
 		GetMutableDefault<UAnimNextConfig>()
 	);
+	
+	FContentBrowserModule& ContentBrowserModule = FModuleManager::LoadModuleChecked<FContentBrowserModule>(TEXT("ContentBrowser"));
+	ContentBrowserModule.RegisterWizard({
+		"UAFAssetWizard",
+		LOCTEXT("UAFAssetWizard_DisplayName","UAF Asset Wizard"),
+		LOCTEXT("UAFAssetWizard_Description", "Wizard for creating UAF assets from templates"),
+		FSlateIcon(FAppStyle::GetAppStyleSetName(), "GraphEditor.Animation_16x"),
+		FAdvancedAssetCategory(EAssetTypeCategories::Animation, LOCTEXT("UAFSubMenu", "Animation Framework")),
+		FSimpleDelegate::CreateLambda([]()
+		{
+			UAF::Editor::FAssetWizard::Launch();
+		})});
 
 #if ANIMNEXT_TRACE_ENABLED
 	IModularFeatures::Get().RegisterModularFeature(IRewindDebuggerExtension::ModularFeatureName, &GRewindDebuggerAnimNext);
@@ -552,6 +567,12 @@ void FAnimNextEditorModule::ShutdownModule()
 #endif // ANIMNEXT_TRACE_ENABLED
 
 	FEdGraphUtilities::UnregisterVisualPinFactory(GraphPanelPinFactory);
+
+	if(FModuleManager::Get().IsModuleLoaded("ContentBrowser"))
+	{
+		FContentBrowserModule& ContentBrowserModule = FModuleManager::LoadModuleChecked<FContentBrowserModule>(TEXT("ContentBrowser"));
+		ContentBrowserModule.UnregisterWizard("UAFAssetWizard");
+	}
 }
 
 void FAnimNextEditorModule::RegisterLocatorFragmentEditorType(FName InLocatorFragmentEditorName)
diff --git a/Engine/Plugins/Experimental/AnimNext/Source/AnimNextEditor/Private/AnimNextStyle.cpp b/Engine/Plugins/Experimental/AnimNext/Source/AnimNextEditor/Private/AnimNextStyle.cpp
--- a/Engine/Plugins/Experimental/AnimNext/Source/AnimNextEditor/Private/AnimNextStyle.cpp
+++ b/Engine/Plugins/Experimental/AnimNext/Source/AnimNextEditor/Private/AnimNextStyle.cpp
@@ -44,6 +44,25 @@ FAnimNextStyle::FAnimNextStyle()
 		.SetPressedPadding(FMargin(6, 3, 6, 1));
 
 	Set("AnimNext.TraitDetails.AddButton", FComboButtonStyle(ComboButtonStyle).SetButtonStyle(AddButtonStyle));
+
+	SetAssetWizardStyles();
+}
+
+void FAnimNextStyle::SetAssetWizardStyles()
+{
+	const FVector2D Icon16x16(16.0f, 16.0f);
+	const FVector2D Icon64x64(64.0f, 64.0f);
+	
+	Set("AssetWizard.OpenExternal", new FSlateVectorImageBrush(RootToContentDir("Icons/Wizard/OpenExternal.svg"),Icon16x16));
+	Set("AssetWizard.Template", new FSlateVectorImageBrush(RootToContentDir("Icons/Wizard/Template.svg"),Icon64x64));
+
+	Set("AssetWizard.TemplateTag.OuterBorder", new FSlateRoundedBoxBrush(FStyleColors::AccentGray, 4.0f));
+	Set("AssetWizard.TemplateTag.InnerBorder", new FSlateRoundedBoxBrush(FStyleColors::Recessed, 4.0f));
+
+	Set("AssetWizard.TemplateFolderBackground", new FSlateRoundedBoxBrush(EStyleColor::Panel, 4.0f));
+	Set("AssetWizard.TemplateFolderHoverBackground", new FSlateRoundedBoxBrush(FStyleColors::Hover, 4.0f));
+	Set("AssetWizard.TemplateFolderSelectedBackground", new FSlateRoundedBoxBrush(FStyleColors::Primary, 4.0f));
+	Set("AssetWizard.TemplateFolderHoverSelectedBackground", new FSlateRoundedBoxBrush(FStyleColors::PrimaryHover, 4.0f));
 }
 
 void FAnimNextStyle::Register()
diff --git a/Engine/Plugins/Experimental/AnimNext/Source/AnimNextEditor/Private/Wizard/AssetWizard.cpp b/Engine/Plugins/Experimental/AnimNext/Source/AnimNextEditor/Private/Wizard/AssetWizard.cpp
new file mode 100644
--- /dev/null
+++ b/Engine/Plugins/Experimental/AnimNext/Source/AnimNextEditor/Private/Wizard/AssetWizard.cpp
@@ -0,0 +1,262 @@
+// Copyright Epic Games, Inc. All Rights Reserved.
+
+#include "AssetWizard.h"
+
+#include "AssetToolsModule.h"
+#include "ContentBrowserModule.h"
+#include "Components/SkeletalMeshComponent.h"
+#include "EditorModeManager.h"
+#include "Engine/SkeletalMesh.h"
+#include "Framework/Application/SlateApplication.h"
+#include "IWorkspaceEditor.h"
+#include "IWorkspaceEditorModule.h"
+#include "TemplateDataAsset.h"
+#include "Component/AnimNextComponent.h"
+#include "Factories/BlueprintFactory.h"
+#include "Module/AnimNextModule.h"
+#include "Serialization/ArchiveReplaceObjectRef.h"
+#include "Wizard/STemplatePicker.h"
+#include "WorkspaceEditor/Private/Workspace.h"
+#include "Engine/SimpleConstructionScript.h"
+#include "Engine/SCS_Node.h"
+#include "Kismet2/KismetEditorUtilities.h"
+#include "Workspace/AnimNextWorkspaceEditorMode.h"
+#include "Workspace/AnimNextWorkspaceFactory.h"
+#include "IWorkspacePicker.h"
+#include "ObjectTools.h"
+#include "TemplateConfig.h"
+
+#define LOCTEXT_NAMESPACE "UE::UAF::Editor::Wizard"
+
+namespace UE::UAF::Editor
+{
+	TObjectPtr<UWorkspace> FAssetWizard::ShowWorkspacePicker()
+	{
+		const Workspace::IWorkspaceEditorModule& WorkspaceEditorModule = FModuleManager::Get().LoadModuleChecked<Workspace::IWorkspaceEditorModule>("WorkspaceEditor");
+
+		Workspace::IWorkspacePicker::FConfig WorkspacePickerConfig;
+		WorkspacePickerConfig.HintText = Workspace::IWorkspacePicker::EHintText::CreateOrUseExistingWorkspace;
+		WorkspacePickerConfig.WorkspaceFactoryClass = UAnimNextWorkspaceFactory::StaticClass();
+		
+		const TSharedPtr<Workspace::IWorkspacePicker> WorkspacePicker = WorkspaceEditorModule.CreateWorkspacePicker(WorkspacePickerConfig);
+		WorkspacePicker->ShowModal();
+
+		return Cast<UWorkspace>(WorkspacePicker->GetSelectedWorkspace());
+	}
+
+	void FAssetWizard::HandleTemplateSelected(const TObjectPtr<const UUAFTemplateDataAsset> Template, const TObjectPtr<const UUAFTemplateConfig> TemplateConfig, const TObjectPtr<UWorkspace> Workspace)
+	{
+		struct FAssetPair
+		{
+			TObjectPtr<UObject> TemplateAsset;
+			TObjectPtr<UObject> UserAsset;
+		};
+		TArray<FAssetPair> Assets;
+
+		bool bFailedToInstantiateTemplateAssets = false;
+
+		for (int32 TemplateAssetIndex = 0; TemplateAssetIndex < Template->Assets.Num(); ++TemplateAssetIndex)
+		{
+			TObjectPtr<UObject> TemplateAsset  = Template->Assets[TemplateAssetIndex];
+			
+			FAssetData AssetData(TemplateAsset);
+
+			const FString AssetName = TemplateConfig->AssetNaming[TemplateAssetIndex];
+			TObjectPtr<UObject> NewAsset = IAssetTools::Get().DuplicateAsset(AssetName, TemplateConfig->OutputPath.Path, TemplateAsset);
+			
+			if (NewAsset != nullptr)
+			{
+				FAssetPair& AssetPair = Assets.AddDefaulted_GetRef();
+				AssetPair.TemplateAsset = TemplateAsset;
+				AssetPair.UserAsset = NewAsset;
+			}
+			else
+			{
+				// Failed to instantiate template
+				bFailedToInstantiateTemplateAssets = true;
+				break;
+			}
+		}
+
+		if (bFailedToInstantiateTemplateAssets)
+		{
+			for (FAssetPair& Asset : Assets)
+			{
+				ObjectTools::DeleteSingleObject(Asset.UserAsset);
+			}
+
+			UE_LOG(LogAnimation, Warning, TEXT("Failed to instantiate template assets"));
+			return;
+		}
+		
+		for (int32 AssetIndex = 0; AssetIndex < Assets.Num(); ++AssetIndex)
+		{
+			// This asset will have all of its references to any of the template assets updates to point to the new user asset.
+			// Any references inside this asset to other assets not part of the template asset pack will remain the same.
+			TObjectPtr<UObject> AssetToRepair = Assets[AssetIndex].UserAsset;
+
+			TMap<UObject*, UObject*> ReplaceMap;
+			
+			for (int32 OtherAssetIndex = 0; OtherAssetIndex < Assets.Num(); ++OtherAssetIndex)
+			{
+				if (AssetIndex != OtherAssetIndex)
+				{
+					FAssetPair OtherAssetPair = Assets[OtherAssetIndex];
+					ReplaceMap.Add(OtherAssetPair.TemplateAsset, OtherAssetPair.UserAsset);
+				}
+			}
+			
+			FArchiveReplaceObjectRef<UObject> ReplaceAr(AssetToRepair, ReplaceMap);
+			AssetToRepair->PostEditChange();
+		}
+
+		for (const FAssetPair& AssetPair : Assets)
+		{
+			Workspace->AddAsset(AssetPair.UserAsset);
+		}
+		
+		if (Template->AssetToOpen)
+		{
+			const int32 AssetIndex = Assets.IndexOfByPredicate([Template](const FAssetPair& InAssetPair){ return InAssetPair.TemplateAsset == Template->AssetToOpen; });
+
+			if (Assets.IsValidIndex(AssetIndex))
+			{
+				TObjectPtr<UObject> AssetToOpen = Assets[AssetIndex].UserAsset;
+				
+				Workspace::IWorkspaceEditorModule& WorkspaceEditorModule = FModuleManager::Get().LoadModuleChecked<Workspace::IWorkspaceEditorModule>("WorkspaceEditor");
+				Workspace::IWorkspaceEditor* WorkspaceEditor = WorkspaceEditorModule.OpenWorkspaceForObject(AssetToOpen, Workspace::EOpenWorkspaceMethod::Default, UAnimNextWorkspaceFactory::StaticClass());
+
+				UAnimNextWorkspaceEditorMode* EditorMode = Cast<UAnimNextWorkspaceEditorMode>(WorkspaceEditor->GetEditorModeManager().GetActiveScriptableMode(UAnimNextWorkspaceEditorMode::EM_AnimNextWorkspace));
+				EditorMode->HandleCompile();
+			}
+		}
+
+		if (TemplateConfig->BlueprintMode == ETemplateBlueprintMode::CreateNewBlueprint || TemplateConfig->BlueprintMode == ETemplateBlueprintMode::ModifyExistingBlueprint)
+		{
+			TObjectPtr<UBlueprint> Blueprint = [TemplateConfig]()
+			{
+				if (TemplateConfig->BlueprintMode == ETemplateBlueprintMode::CreateNewBlueprint)
+				{
+					TObjectPtr<UBlueprintFactory> BlueprintFactory = NewObject<UBlueprintFactory>();
+					BlueprintFactory->AddToRoot();
+			
+					BlueprintFactory->ParentClass = TemplateConfig->BlueprintClass.Get()
+						? TemplateConfig->BlueprintClass.Get()
+						: AActor::StaticClass();
+
+					FAssetToolsModule& AssetToolsModule = FModuleManager::LoadModuleChecked<FAssetToolsModule>("AssetTools");
+					IAssetTools& AssetTools = AssetToolsModule.Get();
+
+					TObjectPtr<UObject> NewAsset = AssetTools.CreateAsset(TemplateConfig->BlueprintAssetName, TemplateConfig->OutputPath.Path, UBlueprint::StaticClass(), BlueprintFactory);
+
+					return CastChecked<UBlueprint>(NewAsset);
+				}
+				else
+				{
+					return TemplateConfig->BlueprintToModify.LoadSynchronous();
+				}
+			}();
+			
+			TObjectPtr<AActor> Actor = CastChecked<AActor>(Blueprint->GeneratedClass->GetDefaultObject());
+
+			USkeletalMeshComponent* SkelMeshComp = Cast<USkeletalMeshComponent>(Actor->GetComponentByClass(USkeletalMeshComponent::StaticClass()));
+			UAnimNextComponent* UAFComp = Cast<UAnimNextComponent>(Actor->GetComponentByClass(UAnimNextComponent::StaticClass()));
+
+			const FAssetPair* ModuleAssetPair = Assets.FindByPredicate([](const FAssetPair& InAssetPair)
+			{
+				return InAssetPair.UserAsset->GetClass() == UAnimNextModule::StaticClass();
+			});
+
+			// Create BP components
+			{
+				USimpleConstructionScript* ConstructionScript = Blueprint->SimpleConstructionScript;
+				USCS_Node* DefaultRootNode = ConstructionScript->GetDefaultSceneRootNode();
+				
+				if (!SkelMeshComp)
+				{
+					// No SkelMesh component exists in this blueprint so create one
+					USCS_Node* SkelMeshCompNode = ConstructionScript->CreateNode(USkeletalMeshComponent::StaticClass(), "Mesh");
+					if (DefaultRootNode)
+					{
+						DefaultRootNode->AddChildNode(SkelMeshCompNode);
+					}
+					else
+					{
+						ConstructionScript->AddNode(SkelMeshCompNode);
+					}
+
+					SkelMeshComp = CastChecked<USkeletalMeshComponent>(SkelMeshCompNode->ComponentTemplate);
+				}
+
+				if (!UAFComp)
+				{
+					// No AnimNext component exists in this blueprint so create one
+					USCS_Node* UAFCompNode = ConstructionScript->CreateNode(UAnimNextComponent::StaticClass(), "AnimationFramework");
+					ConstructionScript->AddNode(UAFCompNode);
+
+					UAFComp = CastChecked<UAnimNextComponent>(UAFCompNode->ComponentTemplate);
+				}
+			}
+
+			// Configure BP components
+			{
+				if (SkelMeshComp)
+				{
+					SkelMeshComp->SetSkeletalMesh(TemplateConfig->SkeletalMesh.LoadSynchronous());
+					SkelMeshComp->SetEnableAnimation(false);
+					
+					FTransform RelativeTransform;
+					RelativeTransform.SetTranslation(FVector(0, 0, -90.0));
+					RelativeTransform.SetRotation(FQuat::MakeFromEuler(FVector(0, 0, -90.0)));
+					SkelMeshComp->SetRelativeTransform(RelativeTransform);
+				}
+
+				if (UAFComp && ModuleAssetPair)
+				{
+					UAFComp->SetModule(CastChecked<UAnimNextModule>(ModuleAssetPair->UserAsset.Get()));
+				}
+			}
+
+			FKismetEditorUtilities::CompileBlueprint(Blueprint);
+		}
+	}
+
+	void FAssetWizard::ShowTemplatePicker()
+	{
+		const TObjectPtr<UWorkspace> Workspace = ShowWorkspacePicker();
+		if (!Workspace)
+		{
+			return;
+		}
+		
+		TSharedRef<SWindow> TemplateWindow = SNew(SWindow)
+			.Title(LOCTEXT("AssetWizard", "UAF Asset Wizard"))
+			.ClientSize(FVector2D(1200, 800))
+			.SupportsMinimize(false)
+			.SupportsMaximize(false);
+
+		TemplateWindow->SetContent(
+			SNew(STemplatePicker)
+			.OnTemplateSelected_Lambda([TemplateWindow, Workspace](const TObjectPtr<const UUAFTemplateDataAsset> InTemplate, const TObjectPtr<const UUAFTemplateConfig> InTemplateConfig) mutable
+			{
+				HandleTemplateSelected(InTemplate, InTemplateConfig, Workspace);
+				FSlateApplication::Get().RequestDestroyWindow(TemplateWindow);
+			})
+			.OnCancel_Lambda([TemplateWindow]()
+			{
+				FSlateApplication::Get().RequestDestroyWindow(TemplateWindow);
+			})
+		);
+
+		GetMutableDefault<UUAFTemplateConfig>()->OutputPath.Path = FAssetData(Workspace).PackagePath.ToString();
+		
+		FSlateApplication::Get().AddWindow(TemplateWindow);
+	}
+	
+	void FAssetWizard::Launch()
+	{
+		ShowTemplatePicker();		
+	}
+}
+
+#undef LOCTEXT_NAMESPACE
\ No newline at end of file
diff --git a/Engine/Plugins/Experimental/AnimNext/Source/AnimNextEditor/Private/Wizard/AssetWizard.h b/Engine/Plugins/Experimental/AnimNext/Source/AnimNextEditor/Private/Wizard/AssetWizard.h
new file mode 100644
--- /dev/null
+++ b/Engine/Plugins/Experimental/AnimNext/Source/AnimNextEditor/Private/Wizard/AssetWizard.h
@@ -0,0 +1,25 @@
+// Copyright Epic Games, Inc. All Rights Reserved.
+
+#pragma once
+
+#include "CoreMinimal.h"
+
+class UUAFTemplateDataAsset;
+class UUAFTemplateConfig;
+class UWorkspace;
+
+namespace UE::UAF::Editor
+{
+	class FAssetWizard
+	{
+	public:
+		static void Launch();
+		
+	private:
+		static TObjectPtr<UWorkspace> ShowWorkspacePicker();
+
+		static void ShowTemplatePicker();
+		
+		static void HandleTemplateSelected(const TObjectPtr<const UUAFTemplateDataAsset> Template, const TObjectPtr<const UUAFTemplateConfig> TemplateConfig, const TObjectPtr<UWorkspace> Workspace);
+	};
+}
\ No newline at end of file
diff --git a/Engine/Plugins/Experimental/AnimNext/Source/AnimNextEditor/Private/Wizard/STemplatePicker.cpp b/Engine/Plugins/Experimental/AnimNext/Source/AnimNextEditor/Private/Wizard/STemplatePicker.cpp
new file mode 100644
--- /dev/null
+++ b/Engine/Plugins/Experimental/AnimNext/Source/AnimNextEditor/Private/Wizard/STemplatePicker.cpp
@@ -0,0 +1,718 @@
+// Copyright Epic Games, Inc. All Rights Reserved.
+
+#include "STemplatePicker.h"
+
+#include "AnimNextStyle.h"
+#include "AssetRegistry/AssetRegistryModule.h"
+#include "AssetThumbnail.h" 
+#include "DetailsViewArgs.h"
+#include "PropertyEditorModule.h"
+#include "SPrimaryButton.h"
+#include "Styling/StyleColors.h"
+#include "Subsystems/EditorAssetSubsystem.h"
+#include "TemplateConfig.h"
+#include "TemplateDataAsset.h"
+#include "ThumbnailRendering/ThumbnailManager.h"
+#include "Widgets/Input/SEditableTextBox.h"
+#include "Widgets/Input/SSegmentedControl.h"
+#include "Widgets/Layout/SExpandableArea.h"
+#include "Widgets/Layout/SScaleBox.h"
+#include "Widgets/Layout/SWidgetSwitcher.h"
+#include "Widgets/Text/STextBlock.h"
+#include "Widgets/Views/STileView.h"
+
+#include "Editor/EditorEngine.h"
+#include "Editor.h"
+
+#define LOCTEXT_NAMESPACE "UE::UAF::Editor::STemplatePicker"
+
+namespace UE::UAF::Editor
+{
+	void STemplatePicker::AddReferencedObjects(FReferenceCollector& Collector)
+	{
+		Collector.AddReferencedObject(TemplateConfig);
+
+		for (FTemplateItemPtr TemplatePtr : Templates)
+		{
+			Collector.AddReferencedObject(TemplatePtr->Template);
+		}
+	}
+
+	void STemplatePicker::Construct(const FArguments& InArgs)
+	{
+		OnTemplateSelected = InArgs._OnTemplateSelected;
+		
+		TSharedRef<SWidget> TemplatesView = SNew(SBorder)
+			.BorderImage(FAppStyle::GetBrush("ToolPanel.GroupBorder"))
+			[
+				SAssignNew(TemplatesTileView, STileView<FTemplateItemPtr>)
+				.ListItemsSource(&Templates)
+				.SelectionMode(ESelectionMode::Single)
+				.ClearSelectionOnClick(true)
+				.ItemAlignment(EListItemAlignment::LeftAligned)
+				.OnGenerateTile_Lambda([](FTemplateItemPtr InItem, const TSharedRef<STableViewBase>& OwnerTable)
+				{
+					return SNew(SModuleTemplate, OwnerTable)
+						.Item(InItem);
+				})
+				.ItemWidth(111)
+				.ItemHeight(162)
+				.OnSelectionChanged_Lambda([this](FTemplateItemPtr InItem, ESelectInfo::Type SelectInfo)
+				{
+					SelectedTemplate = InItem;
+					HandleSelectedTemplateChanged();
+				})
+			];
+
+		TemplateConfig = GetMutableDefault<UUAFTemplateConfig>();
+		TemplateConfig->OnOutputPathChanged.AddSP(this, &STemplatePicker::HandleOutputPathChanged);
+		
+		FDetailsViewArgs DetailsViewArgs;
+		DetailsViewArgs.bAllowSearch = false;
+		DetailsViewArgs.bLockable = false;
+		DetailsViewArgs.bShowOptions = false;
+		DetailsViewArgs.bAllowFavoriteSystem = false;
+		DetailsViewArgs.bHideSelectionTip = true;
+
+		FPropertyEditorModule& PropertyEditorModule = FModuleManager::Get().GetModuleChecked<FPropertyEditorModule>("PropertyEditor");
+		TSharedRef<IDetailsView> DetailsObjView = PropertyEditorModule.CreateDetailView(DetailsViewArgs);
+		DetailsObjView->SetObject(TemplateConfig);
+
+		TSharedRef<SWidget> DetailsView = SNew(SBorder)
+			.BorderImage(FAppStyle::GetBrush("ToolPanel.GroupBorder"))
+			[
+				SNew(SVerticalBox)
+				// Panel selector
+				+ SVerticalBox::Slot()
+				.HAlign(HAlign_Center)
+				.AutoHeight()
+				[
+					SNew(SHorizontalBox)
+					+ SHorizontalBox::Slot()
+					.AutoWidth()
+					[
+						SNew(SSegmentedControl<EPanel>)
+						.IsEnabled_Lambda([this]()
+						{
+							return SelectedTemplate.IsValid();
+						})
+						.UniformPadding(FMargin(8.0f, 0.0f))
+						.OnValueChanged_Lambda([this](const EPanel InPanel)
+						{
+							ShowPanel = InPanel;
+						})
+						.Value_Lambda([this]()
+						{
+							return ShowPanel;
+						})	
+						+ SSegmentedControl<EPanel>::Slot(EPanel::Overview)
+						[
+							SNew(SHorizontalBox)
+							+ SHorizontalBox::Slot()
+							.Padding(2.0f)
+							.AutoWidth()
+							[
+								SNew(SImage)
+								.Image(FAppStyle::GetBrush("GraphEditor.Animation_16x"))
+							]
+							+ SHorizontalBox::Slot()
+							.AutoWidth()
+							.Padding(2.0f)
+							[
+								SNew(STextBlock)
+								.Text(LOCTEXT("Overview", "Overview"))
+							]
+						]
+						+ SSegmentedControl<EPanel>::Slot(EPanel::Details)
+						[
+							SNew(SHorizontalBox)
+							+ SHorizontalBox::Slot()
+							.Padding(2.0f)
+							.AutoWidth()
+							[
+								SNew(SImage)
+								.Image(FAppStyle::GetBrush("LevelEditor.Tabs.Details"))
+							]
+							+ SHorizontalBox::Slot()
+							.AutoWidth()
+							.Padding(2.0f)
+							[
+								SNew(STextBlock)
+								.Text(LOCTEXT("Details", "Details"))
+							]
+						]
+					]
+				]
+				// Panel
+				+ SVerticalBox::Slot()
+				.FillHeight(1.0f)
+				[
+					SNew(SWidgetSwitcher)
+					.WidgetIndex_Lambda([this]()
+					{
+						return static_cast<int32>(ShowPanel);
+					})
+					+ SWidgetSwitcher::Slot() // EPanel::Overview
+					.Padding(8.0f)
+					[
+						SNew(SVerticalBox)
+						// Preview image
+						+ SVerticalBox::Slot()
+						.HAlign(HAlign_Center)
+						.AutoHeight()
+						[
+							SNew(SScaleBox)
+							.Stretch(EStretch::ScaleToFill)
+							[
+								SNew(SImage)
+								.DesiredSizeOverride(FVector2D(600,300))
+								.Image_Lambda([this]() -> const FSlateBrush*
+								{
+									if (SelectedTemplate && SelectedTemplate->Template->DetailsImage.IsSet())
+									{
+										return &SelectedTemplate->Template->DetailsImage;
+									}
+
+									return FAppStyle::GetNoBrush();
+								})
+							]
+						]
+						// Template Name
+						+ SVerticalBox::Slot()
+						.AutoHeight()
+						.Padding(FMargin(0.0f, 16.0f, 0.0f, 2.0f))
+						[
+							SNew(STextBlock)
+							.AutoWrapText(true)
+							.TextStyle(FAppStyle::Get(), "LargeText")
+							.Text_Lambda([this]()
+							{
+								return SelectedTemplate.IsValid()
+									? SelectedTemplate->Template->Title
+									: FText();
+							})
+						]
+						// Template Description
+						+ SVerticalBox::Slot()
+						.AutoHeight()
+						.Padding(FMargin(0.0f, 2.0f, 0.0f, 0.0f))
+						[
+							SNew(STextBlock)
+							.AutoWrapText(true)
+							.Text_Lambda([this]()
+							{
+								return SelectedTemplate.IsValid()
+									? SelectedTemplate->Template->Description
+									: FText();
+							})
+						]
+						// Documentation URL
+						+ SVerticalBox::Slot()
+						.Padding(0.0f, 24.0f)
+						.AutoHeight()
+						[
+							SNew(SHorizontalBox)
+							+ SHorizontalBox::Slot()
+							.AutoWidth()
+							[
+								SNew(SButton)
+								.Cursor(EMouseCursor::Hand)
+								.Visibility_Lambda([this]()
+								{
+									return SelectedTemplate.IsValid() && !SelectedTemplate->Template->DocumentationUrl.IsEmpty()
+										? EVisibility::Visible
+										: EVisibility::Collapsed;
+								})
+								.ButtonStyle(FCoreStyle::Get(), "NoBorder")
+								.OnClicked_Lambda([this]()
+								{
+									if (SelectedTemplate.IsValid() && !SelectedTemplate->Template->DocumentationUrl.IsEmpty())
+									{
+										FPlatformProcess::LaunchURL(*SelectedTemplate->Template->DocumentationUrl, nullptr, nullptr);
+									}
+									return FReply::Handled();
+								})
+								.ContentPadding(0)
+								[
+									SNew(SHorizontalBox)
+
+									+ SHorizontalBox::Slot()
+									.AutoWidth()
+									[
+										SNew(STextBlock)
+										.Text(LOCTEXT("Documentation", "Documentation"))
+										.ColorAndOpacity(FLinearColor(0.2f, 0.4f, 1.0f))
+									]
+
+									+ SHorizontalBox::Slot()
+									.AutoWidth()
+									.Padding(FMargin(4, 0, 0, 0))
+									[
+										SNew(SImage)
+										.Image(FAnimNextStyle::Get().GetBrush("AssetWizard.OpenExternal"))
+										.ColorAndOpacity(FLinearColor(0.2f, 0.4f, 1.0f))
+									]
+								]
+							]
+						]
+						// Tags
+						+ SVerticalBox::Slot()
+						.AutoHeight()
+						[
+							SAssignNew(TagsHBox, SHorizontalBox)
+						]
+					]
+					+ SWidgetSwitcher::Slot() // EPanel::Details
+					[
+						SNew(SVerticalBox)
+						.Visibility_Lambda([this]()
+						{
+							return SelectedTemplate.IsValid() ? EVisibility::Visible : EVisibility::Hidden;		
+						})
+						+ SVerticalBox::Slot()
+						.AutoHeight()
+						[
+							DetailsObjView
+						]
+						+ SVerticalBox::Slot()
+						.AutoHeight()
+						[
+							SNew(SExpandableArea)
+							.AreaTitle(LOCTEXT("IncludedAssets", "Included Assets"))
+							.BorderImage(FAppStyle::GetBrush("Brushes.Header"))
+							.BodyBorderImage(FAppStyle::GetBrush("Brushes.Recessed"))
+							.HeaderPadding(4.0f)
+							.Padding(2.0f)
+							.AllowAnimatedTransition(false)
+							.InitiallyCollapsed(false)
+							.BodyContent()
+							[
+								SAssignNew(IncludedAssetsListView, SListView<TSharedPtr<FAssetMetadata>>)
+								.ListItemsSource(&IncludedAssets)
+								.OnGenerateRow_Lambda([this](TSharedPtr<FAssetMetadata> InAsset, const TSharedRef<STableViewBase>& InTableView)
+								{
+									TSharedPtr<FAssetThumbnailPool> Pool = UThumbnailManager::Get().GetSharedThumbnailPool();
+									TSharedPtr<FAssetThumbnail> AssetThumbnail = MakeShareable(new FAssetThumbnail(InAsset->AssetData.GetAsset(), 32, 32, Pool));
+
+									FAssetThumbnailConfig AssetThumbnailConfig;
+									AssetThumbnailConfig.ShowAssetBorder = true;
+																	
+									return SNew(STableRow<TSharedPtr<FAssetMetadata>>, InTableView)
+									.Content()
+									[
+										SNew(SHorizontalBox)
+										+ SHorizontalBox::Slot()
+										.FillWidth(0.5f)
+										.Padding(4.0f)
+										[
+											SNew(SHorizontalBox)
+											+ SHorizontalBox::Slot()
+											.AutoWidth()
+											[
+												AssetThumbnail->MakeThumbnailWidget(AssetThumbnailConfig)
+
+											]
+											+ SHorizontalBox::Slot()
+											.AutoWidth()
+											.VAlign(VAlign_Center)
+											.Padding(4.0f)
+											[
+												SNew(STextBlock)
+												.Text(FText::FromName(InAsset->AssetData.AssetName))
+											]
+										]
+										+ SHorizontalBox::Slot()
+										.FillWidth(0.5f)
+										.Padding(4.0f)
+										.VAlign(VAlign_Center)
+										[
+											SNew(SEditableTextBox)
+											.Text_Lambda([InAsset]
+											{
+												return FText::FromString(InAsset->AssetRename);
+											})
+											.OnTextCommitted_Lambda([InAsset](const FText& InText, ETextCommit::Type)
+											{
+												InAsset->AssetRename = InText.ToString();
+											})
+											.OnVerifyTextChanged_Lambda([this](const FText& InNewText, FText& OutErrorMessage)
+											{
+												UEditorAssetSubsystem* EditorAssetSubsystem = GEditor->GetEditorSubsystem<UEditorAssetSubsystem>();
+
+												const FString AssetPath = FPaths::Combine( TemplateConfig->OutputPath.Path, InNewText.ToString());
+
+												if (EditorAssetSubsystem->DoesAssetExist(AssetPath))
+												{
+													OutErrorMessage = LOCTEXT("AssetAlreadyExists", "An asset with that name already exists");
+													return false;
+												}
+												
+												return true;
+											})
+										]
+									];
+								})
+							]
+						]
+					]
+				]
+			];
+
+		TSharedRef<SWidget> HeaderWidget = SNew(SBorder)
+			.BorderImage(FAppStyle::GetBrush("ToolPanel.GroupBorder"))
+			.Padding(16.0f)
+			[
+				SNew(STextBlock)
+				.Text(LOCTEXT("Header","Create a bundle preconfigured UAF assets by selecting a template to generate from."))
+			];
+
+		TSharedRef<SWidget> FooterWidget = SNew(SBorder)
+			.BorderImage(FAppStyle::GetBrush("ToolPanel.GroupBorder"))
+			[
+				SNew(SHorizontalBox)
+				+ SHorizontalBox::Slot()
+				.AutoWidth()
+				.Padding(16.0f)
+				[
+					SNullWidget::NullWidget
+				]
+				+ SHorizontalBox::Slot()
+				.FillWidth(1.0f)
+				.HAlign(HAlign_Right)
+				[
+					SNew(SHorizontalBox)
+					+ SHorizontalBox::Slot()
+					.HAlign(HAlign_Right)
+					.AutoWidth()
+					.Padding(FMargin(16.0f, 16.0f, 8.0f, 16.0f))
+					[
+						SNew(SPrimaryButton)
+						.IsEnabled_Lambda([this]()
+						{
+							return SelectedTemplate.IsValid();
+						})
+						.Text(LOCTEXT("CreateAssets", "Create Assets"))
+						.OnClicked_Lambda([this]()
+						{
+							HandleUseTemplate();
+							return FReply::Handled();
+						})
+					]
+					+ SHorizontalBox::Slot()
+					.HAlign(HAlign_Right)
+					.AutoWidth()
+					.Padding(FMargin(8.0f, 16.0f, 16.0f, 16.0f))
+					[
+						SNew(SButton)
+						.Text(LOCTEXT("Cancel","Cancel"))
+						.OnClicked_Lambda([InArgs]()
+						{
+							InArgs._OnCancel.ExecuteIfBound();
+							return FReply::Handled();
+						})
+					]
+				]
+			];
+		
+		ChildSlot
+		[
+			SNew(SVerticalBox)
+			+ SVerticalBox::Slot()
+			.Padding(8.0f, 8.0f, 8.0f, 0.0f)
+			.AutoHeight()
+			[
+				HeaderWidget
+			]
+			+ SVerticalBox::Slot()
+			.Padding(4.0f)
+			.FillHeight(1.0f)
+			[
+				SNew(SHorizontalBox)
+				+ SHorizontalBox::Slot()
+				.Padding(4.0f)
+				.FillWidth(1.0f)
+				[
+					SNew(SSplitter)
+					.Orientation(Orient_Horizontal)
+					+ SSplitter::Slot()
+					.Value(0.5)
+					[
+						TemplatesView
+					]
+					+ SSplitter::Slot()
+					.Value(0.5)
+					[
+						DetailsView
+					]
+				]
+			]
+			+ SVerticalBox::Slot()
+			.Padding(8.0f, 0.0f, 8.0f, 8.0f)
+			.AutoHeight()
+			[
+				FooterWidget
+			]
+		];
+
+		RefreshModuleTemplates();
+
+		FTemplateItemPtr* DefaultSelectedTemplate = Templates.FindByPredicate([](const FTemplateItemPtr TemplateItem)
+		{
+			return TemplateItem->Template->GetPathName() == TEXT("/AnimNext/Templates/BasicCharacter/DA_DefaultCharacter.DA_DefaultCharacter");
+		});
+
+		if (DefaultSelectedTemplate)
+		{
+			SelectedTemplate = *DefaultSelectedTemplate;
+			TemplatesTileView->SetSelection(SelectedTemplate);
+		}
+		
+		HandleSelectedTemplateChanged();
+	}
+		
+	void STemplatePicker::HandleUseTemplate()
+	{
+		bool bSuccess = TemplateConfig->TryUpdateDefaultConfigFile();
+
+		for (const TSharedPtr<FAssetMetadata>& AssetMetadata : IncludedAssets)
+		{
+			TemplateConfig->AssetNaming.Add(AssetMetadata->AssetRename);
+		}
+		
+		OnTemplateSelected.ExecuteIfBound(SelectedTemplate->Template, TemplateConfig);
+	}
+
+	void STemplatePicker::HandleOutputPathChanged()
+	{
+		RebuildIncludedAssetsListView();
+	}
+
+	void STemplatePicker::RebuildIncludedAssetsListView()
+	{
+		UEditorAssetSubsystem* EditorAssetSubsystem = GEditor->GetEditorSubsystem<UEditorAssetSubsystem>();
+		
+		IncludedAssets.Empty();
+
+		if (SelectedTemplate.IsValid())
+		{
+			check(SelectedTemplate->Template);
+			
+			for (const TObjectPtr<UObject>& Asset : SelectedTemplate->Template->Assets)
+			{
+				TSharedPtr<FAssetMetadata> AssetMetadata = MakeShared<FAssetMetadata>();
+				AssetMetadata->AssetData = FAssetData(Asset);
+				
+				const FString BaseAssetName = AssetMetadata->AssetData.AssetName.ToString();
+				AssetMetadata->AssetRename = BaseAssetName;
+				
+				FString NewAssetPath = FPaths::Combine(TemplateConfig->OutputPath.Path, AssetMetadata->AssetRename);
+				if (EditorAssetSubsystem->DoesAssetExist(NewAssetPath))
+				{
+					int32 Suffix = 1;
+					
+					while (true)
+					{
+						AssetMetadata->AssetRename = BaseAssetName + "_" + FString::FromInt(Suffix++);
+						NewAssetPath = FPaths::Combine(TemplateConfig->OutputPath.Path, AssetMetadata->AssetRename);
+
+						if (!EditorAssetSubsystem->DoesAssetExist(NewAssetPath))
+						{
+							break;
+						}
+					}
+				}
+				
+				IncludedAssets.Add(AssetMetadata);
+			}
+		}
+
+		IncludedAssetsListView->RequestListRefresh();
+	}
+	
+	void STemplatePicker::HandleSelectedTemplateChanged()
+	{
+		TagsHBox->ClearChildren();
+
+		if (SelectedTemplate.IsValid())
+		{
+			TemplateConfig->BlueprintAssetName = SelectedTemplate->Template->DefaultBlueprintAssetName;
+			
+			for (const FText& InTag : SelectedTemplate->Template->Tags)
+			{
+				TagsHBox->AddSlot()
+				.Padding(0.0f, 4.0f, 4.0f, 4.0f)
+				.AutoWidth()
+				[
+					SNew(SBorder)
+					.BorderImage(FAnimNextStyle::Get().GetBrush("AssetWizard.TemplateTag.OuterBorder"))
+					.BorderBackgroundColor(FLinearColor::White)
+					.Padding(1.0f)
+					[
+						SNew(SBorder)
+						.BorderImage(FAnimNextStyle::Get().GetBrush("AssetWizard.TemplateTag.InnerBorder"))
+						.Padding(12.0f, 2.0f)
+						[
+							SNew(STextBlock)
+							.Text(InTag)
+						]
+					]
+				];
+			}
+		}
+		
+		RebuildIncludedAssetsListView();
+	}
+	
+	void STemplatePicker::RefreshModuleTemplates()
+	{
+		FAssetRegistryModule& AssetRegistryModule = FModuleManager::LoadModuleChecked<FAssetRegistryModule>("AssetRegistry");
+		IAssetRegistry& AssetRegistry = AssetRegistryModule.Get();
+
+		TArray<FAssetData> TemplateAssets;
+		AssetRegistry.GetAssetsByClass(UUAFTemplateDataAsset::StaticClass()->GetClassPathName(), TemplateAssets);
+
+		for (const FAssetData& AssetData : TemplateAssets)
+		{
+			TObjectPtr<UUAFTemplateDataAsset> ModuleTemplateDataAsset = CastChecked<UUAFTemplateDataAsset>(AssetData.GetAsset());
+
+			TSharedPtr<FTemplate> ModuleTemplate = Templates.Add_GetRef(MakeShared<FTemplate>(ModuleTemplateDataAsset));
+			ModuleTemplate->Template = ModuleTemplateDataAsset;
+		}
+
+		TemplatesTileView->RequestListRefresh();
+	}
+
+	void SModuleTemplate::Construct(const FArguments& InArgs, const TSharedRef<STableViewBase>& OwnerTable)
+	{
+		Item = InArgs._Item;
+
+		STableRow::FArguments TableRowArguments;
+		TableRowArguments._SignalSelectionMode = ETableRowSignalSelectionMode::Instantaneous;
+		
+		STableRow::Construct(
+			TableRowArguments
+			.Style(FAppStyle::Get(), "ProjectBrowser.TableRow")
+			.Padding(2.0f)
+			.Content()
+			[
+				SNew(SBorder)
+				.BorderImage(FAppStyle::Get().GetBrush("ProjectBrowser.ProjectTile.DropShadow"))
+				.Clipping(EWidgetClipping::ClipToBounds)
+				[
+					SNew(SOverlay)
+					+ SOverlay::Slot()
+					[
+						SNew(SVerticalBox)
+						+ SVerticalBox::Slot() // Thumbnail
+						.AutoHeight()
+						.HAlign(HAlign_Center)
+						.VAlign(VAlign_Center)
+						[
+							SNew(SBox)
+							.WidthOverride(102.0f)
+							.HeightOverride(102.0f)
+							[
+								SNew(SBorder)
+								.Padding(0.0f)
+								.BorderImage(FAppStyle::Get().GetBrush("ProjectBrowser.ProjectTile.ThumbnailAreaBackground"))
+								.HAlign(HAlign_Center)
+								.VAlign(VAlign_Center)
+								[
+									SNew(SImage)
+									.Image(FAnimNextStyle::Get().GetBrush("AssetWizard.Template"))
+								]
+							]
+						]
+						// Name
+						+ SVerticalBox::Slot()
+						[
+							SNew(SBorder)
+							.VAlign(VAlign_Top)
+							.Padding(FMargin(3.0f, 3.0f))
+							.BorderImage_Lambda
+							(
+								[this]()
+								{
+									const bool bIsSelected = IsSelected();
+									const bool bIsRowHovered = IsHovered();
+
+									if (bIsSelected && bIsRowHovered)
+									{
+										static const FName SelectedHover("ProjectBrowser.ProjectTile.NameAreaSelectedHoverBackground");
+										return FAppStyle::Get().GetBrush(SelectedHover);
+									}
+									if (bIsSelected)
+									{
+										static const FName Selected("ProjectBrowser.ProjectTile.NameAreaSelectedBackground");
+										return FAppStyle::Get().GetBrush(Selected);
+									}
+									if (bIsRowHovered)
+									{
+										static const FName Hovered("ProjectBrowser.ProjectTile.NameAreaHoverBackground");
+										return FAppStyle::Get().GetBrush(Hovered);
+									}
+
+									return FAppStyle::Get().GetBrush("ProjectBrowser.ProjectTile.NameAreaBackground");
+								}
+							)
+							[
+								SNew(STextBlock)
+								.Font(FAppStyle::Get().GetFontStyle("ProjectBrowser.ProjectTile.Font"))
+								.AutoWrapText(true)
+								.Text(Item->Template->Title)
+								.ColorAndOpacity_Lambda
+								(
+									[this]()
+									{
+										const bool bIsSelected = IsSelected();
+										const bool bIsRowHovered = IsHovered();
+
+										if (bIsSelected || bIsRowHovered)
+										{
+											return FStyleColors::White;
+										}
+
+										return FSlateColor::UseForeground();
+									}
+								)
+							]
+						]
+					]
+					+ SOverlay::Slot()
+					[
+						SNew(SImage)
+						.Visibility(EVisibility::HitTestInvisible)
+						.Image_Lambda
+						(
+							[this]()
+							{
+								const bool bIsSelected = IsSelected();
+								const bool bIsRowHovered = IsHovered();
+
+								if (bIsSelected && bIsRowHovered)
+								{
+									static const FName SelectedHover("ProjectBrowser.ProjectTile.SelectedHoverBorder");
+									return FAppStyle::Get().GetBrush(SelectedHover);
+								}
+								if (bIsSelected)
+								{
+									static const FName Selected("ProjectBrowser.ProjectTile.SelectedBorder");
+									return FAppStyle::Get().GetBrush(Selected);
+								}
+								if (bIsRowHovered)
+								{
+									static const FName Hovered("ProjectBrowser.ProjectTile.HoverBorder");
+									return FAppStyle::Get().GetBrush(Hovered);
+								}
+
+								return FStyleDefaults::GetNoBrush();
+							}
+						)
+					]
+				]
+			],
+			OwnerTable);
+	}
+}
+
+#undef LOCTEXT_NAMESPACE
diff --git a/Engine/Plugins/Experimental/AnimNext/Source/AnimNextEditor/Private/Wizard/STemplatePicker.h b/Engine/Plugins/Experimental/AnimNext/Source/AnimNextEditor/Private/Wizard/STemplatePicker.h
new file mode 100644
--- /dev/null
+++ b/Engine/Plugins/Experimental/AnimNext/Source/AnimNextEditor/Private/Wizard/STemplatePicker.h
@@ -0,0 +1,98 @@
+// Copyright Epic Games, Inc. All Rights Reserved.
+
+#pragma once
+
+#include "AssetRegistry/AssetData.h"
+#include "CoreMinimal.h"
+#include "TemplateDataAsset.h"
+#include "Widgets/SCompoundWidget.h"
+#include "Widgets/Views/STileView.h"
+
+class UUAFTemplateDataAsset;
+class UUAFTemplateConfig;
+
+namespace UE::UAF::Editor
+{
+	DECLARE_DELEGATE_TwoParams(FOnTemplateSelected, const TObjectPtr<const UUAFTemplateDataAsset>, const TObjectPtr<const UUAFTemplateConfig>);
+
+	class STemplatePicker : public SCompoundWidget, public FGCObject
+	{
+	public:
+		friend class SModuleTemplate;
+		
+		SLATE_BEGIN_ARGS(STemplatePicker) {}
+			SLATE_EVENT(FOnTemplateSelected, OnTemplateSelected)
+			SLATE_EVENT(FSimpleDelegate, OnCancel)
+		SLATE_END_ARGS()
+
+		/** FGCObject interface */
+		virtual void AddReferencedObjects( FReferenceCollector& Collector ) override;
+		virtual FString GetReferencerName() const override
+		{
+			return TEXT("STemplatePicker");
+		}
+
+		void Construct(const FArguments& InArgs);
+
+	private:
+		void RefreshModuleTemplates();
+
+		void HandleSelectedTemplateChanged();
+
+		void HandleUseTemplate();
+		
+		void HandleOutputPathChanged();
+
+		void RebuildIncludedAssetsListView();
+		
+		enum class EPanel
+		{
+			Overview,
+			Details
+		};
+
+		struct FTemplate
+		{
+			TObjectPtr<UUAFTemplateDataAsset> Template;
+		};
+	
+		using FTemplateItemPtr = TSharedPtr<FTemplate>;
+		
+		TSharedPtr<STileView<FTemplateItemPtr>> TemplatesTileView;
+
+		TArray<FTemplateItemPtr> Templates;
+
+		FTemplateItemPtr SelectedTemplate;
+
+		TSharedPtr<SHorizontalBox> TagsHBox;
+		
+		EPanel ShowPanel = EPanel::Details;
+
+		struct FAssetMetadata
+		{
+			FString AssetRename;
+			FAssetData AssetData;
+		};
+		
+		TArray<TSharedPtr<FAssetMetadata>> IncludedAssets;
+		
+		FOnTemplateSelected OnTemplateSelected;
+
+		TObjectPtr<UUAFTemplateConfig> TemplateConfig;
+		
+		TSharedPtr<SListView<TSharedPtr<FAssetMetadata>>> IncludedAssetsListView;
+	};
+
+	class SModuleTemplate : public STableRow<TSharedPtr<STemplatePicker::FTemplate>>
+	{
+	public:
+		SLATE_BEGIN_ARGS(SModuleTemplate) {}
+		SLATE_ARGUMENT(TSharedPtr<STemplatePicker::FTemplate>, Item)
+	SLATE_END_ARGS()
+
+	void Construct(const FArguments& InArgs, const TSharedRef<STableViewBase>& OwnerTable);
+
+	private:
+		TSharedPtr<STemplatePicker::FTemplate> Item;
+	};
+}
\ No newline at end of file
diff --git a/Engine/Plugins/Experimental/AnimNext/Source/AnimNextEditor/Private/Wizard/TemplateConfig.cpp b/Engine/Plugins/Experimental/AnimNext/Source/AnimNextEditor/Private/Wizard/TemplateConfig.cpp
new file mode 100644
--- /dev/null
+++ b/Engine/Plugins/Experimental/AnimNext/Source/AnimNextEditor/Private/Wizard/TemplateConfig.cpp
@@ -0,0 +1,13 @@
+// Copyright Epic Games, Inc. All Rights Reserved.
+
+#include "TemplateConfig.h"
+
+void UUAFTemplateConfig::PostEditChangeProperty(FPropertyChangedEvent& PropertyChangedEvent)
+{
+	UObject::PostEditChangeProperty(PropertyChangedEvent);
+
+	if (PropertyChangedEvent.GetMemberPropertyName() == GET_MEMBER_NAME_CHECKED(UUAFTemplateConfig, OutputPath))
+	{
+		OnOutputPathChanged.Broadcast();
+	}
+}
\ No newline at end of file
diff --git a/Engine/Plugins/Experimental/AnimNext/Source/AnimNextEditor/Private/Wizard/TemplateConfig.h b/Engine/Plugins/Experimental/AnimNext/Source/AnimNextEditor/Private/Wizard/TemplateConfig.h
new file mode 100644
--- /dev/null
+++ b/Engine/Plugins/Experimental/AnimNext/Source/AnimNextEditor/Private/Wizard/TemplateConfig.h
@@ -0,0 +1,54 @@
+// Copyright Epic Games, Inc. All Rights Reserved.
+
+#pragma once
+
+#include "CoreMinimal.h"
+#include "GameFramework/Pawn.h"
+#include "Templates/SubclassOf.h"
+
+#include "TemplateConfig.generated.h"
+
+class USkeletalMesh;
+class UBlueprint;
+
+UENUM()
+enum class ETemplateBlueprintMode
+{
+	CreateNewBlueprint,
+	ModifyExistingBlueprint,
+	DoNothing,
+};
+
+UCLASS(Config = EditorPerProjectUserSettings)
+class UUAFTemplateConfig :	public UObject
+{
+	GENERATED_BODY()
+
+public:
+	// Whether the template should also set up a new Blueprint for this template
+	UPROPERTY(Config, EditAnywhere, Category = "Config")
+	ETemplateBlueprintMode BlueprintMode = ETemplateBlueprintMode::CreateNewBlueprint;
+
+	UPROPERTY(Config, EditAnywhere, Category = "Config", Meta = (EditCondition = "BlueprintMode == ETemplateBlueprintMode::CreateNewBlueprint", EditConditionHides))
+	FString BlueprintAssetName;
+	
+	UPROPERTY(Config, EditAnywhere, Category = "Config", Meta = (EditCondition = "BlueprintMode == ETemplateBlueprintMode::CreateNewBlueprint", EditConditionHides))
+	TSubclassOf<AActor> BlueprintClass = APawn::StaticClass();
+	
+	UPROPERTY(Config, EditAnywhere, Category = "Config", Meta = (EditCondition = "BlueprintMode == ETemplateBlueprintMode::CreateNewBlueprint", EditConditionHides))
+	TSoftObjectPtr<USkeletalMesh> SkeletalMesh;
+
+	UPROPERTY(EditAnywhere, Category = "Config", Meta = (EditCondition = "BlueprintMode == ETemplateBlueprintMode::ModifyExistingBlueprint", EditConditionHides, DisallowCreateNew))
+	TSoftObjectPtr<UBlueprint> BlueprintToModify;
+
+	UPROPERTY(EditAnywhere, Category = "Output", Meta = (ContentDir))
+	FDirectoryPath OutputPath;
+
+	// UObject
+	virtual void PostEditChangeProperty(FPropertyChangedEvent& PropertyChangedEvent) override;
+
+	FSimpleMulticastDelegate OnOutputPathChanged;
+
+public:
+	TArray<FString> AssetNaming;
+};
\ No newline at end of file
diff --git a/Engine/Plugins/Experimental/AnimNext/Source/AnimNextEditor/Private/Wizard/TemplateDataAsset.h b/Engine/Plugins/Experimental/AnimNext/Source/AnimNextEditor/Private/Wizard/TemplateDataAsset.h
new file mode 100644
--- /dev/null
+++ b/Engine/Plugins/Experimental/AnimNext/Source/AnimNextEditor/Private/Wizard/TemplateDataAsset.h
@@ -0,0 +1,43 @@
+// Copyright Epic Games, Inc. All Rights Reserved.
+
+#pragma once
+
+#include "CoreMinimal.h"
+#include "Engine/DataAsset.h"
+#include "Styling/SlateBrush.h"
+
+#include "TemplateDataAsset.generated.h"
+
+UCLASS()
+class UUAFTemplateDataAsset : public UPrimaryDataAsset
+{
+	GENERATED_BODY()
+
+public:
+	UPROPERTY(EditAnywhere, Category = Default)
+	FText Title;
+
+	UPROPERTY(EditAnywhere, Category = Default)
+	FText Description;
+	
+	UPROPERTY(EditAnywhere, Category = Default)
+	FString DocumentationUrl;
+
+	UPROPERTY(EditAnywhere, Category = Default)
+	TArray<FText> Tags;
+
+	UPROPERTY(EditAnywhere, Category = Default)
+	FSlateBrush ThumbnailImage;
+
+	UPROPERTY(EditAnywhere, Category = Default)
+	FSlateBrush DetailsImage;
+
+	UPROPERTY(EditAnywhere, Category = Default)
+	TArray<TObjectPtr<UObject>> Assets;
+
+	UPROPERTY(EditAnywhere, Category = Default)
+	TObjectPtr<UObject> AssetToOpen;
+
+	UPROPERTY(EditAnywhere, Category = Default)
+	FString DefaultBlueprintAssetName;
+};
\ No newline at end of file
diff --git a/Engine/Plugins/Experimental/AnimNext/Source/AnimNextEditor/Private/Workspace/AnimNextWorkspaceEditorMode.h b/Engine/Plugins/Experimental/AnimNext/Source/AnimNextEditor/Private/Workspace/AnimNextWorkspaceEditorMode.h
--- a/Engine/Plugins/Experimental/AnimNext/Source/AnimNextEditor/Private/Workspace/AnimNextWorkspaceEditorMode.h
+++ b/Engine/Plugins/Experimental/AnimNext/Source/AnimNextEditor/Private/Workspace/AnimNextWorkspaceEditorMode.h
@@ -1,4 +1,4 @@
-﻿// Copyright Epic Games, Inc. All Rights Reserved.
+// Copyright Epic Games, Inc. All Rights Reserved.
 
 #pragma once
 
@@ -19,6 +19,11 @@ struct FWorkspaceDocument;
 class IWorkspaceEditor;
 }
 
+namespace UE::UAF::Editor
+{
+	class FAssetWizard;
+}
+
 UCLASS(Transient)
 class UAnimNextWorkspaceEditorMode : public UEdMode
 {
@@ -42,7 +47,6 @@ private:
 	virtual void Exit() override;
 	virtual void BindCommands() override;
 	virtual bool UsesToolkits() const override { return false; }
-
 	void HandleCompile();
 	void HandleAutoCompile();
 	bool IsAutoCompileChecked() const;
@@ -72,6 +76,7 @@ private:
 	TSharedPtr<UE::UAF::Editor::IAssetCompilationHandler> GetAssetCompiler(UObject* InAsset);
 
 	friend UAnimNextWorkspaceSchema;
+	friend UE::UAF::Editor::FAssetWizard;
 
 	// Custom state, persisted via UAnimNextWorkspaceSchema
 	FAnimNextWorkspaceState State;
diff --git a/Engine/Plugins/Experimental/Workspace/Source/WorkspaceEditor/Private/SWorkspacePicker.cpp b/Engine/Plugins/Experimental/Workspace/Source/WorkspaceEditor/Private/SWorkspacePicker.cpp
--- a/Engine/Plugins/Experimental/Workspace/Source/WorkspaceEditor/Private/SWorkspacePicker.cpp
+++ b/Engine/Plugins/Experimental/Workspace/Source/WorkspaceEditor/Private/SWorkspacePicker.cpp
@@ -1,10 +1,13 @@
 ﻿// Copyright Epic Games, Inc. All Rights Reserved.
 
 #include "SWorkspacePicker.h"
+
+#include "AssetToolsModule.h"
 #include "Workspace.h"
 #include "ContentBrowserModule.h"
 #include "IContentBrowserSingleton.h"
 #include "SPrimaryButton.h"
+#include "WorkspaceFactory.h"
 #include "Widgets/Layout/SBox.h"
 #include "Widgets/Text/STextBlock.h"
 #include "Framework/Application/SlateApplication.h"
@@ -17,14 +20,17 @@ namespace UE::Workspace
 void SWorkspacePicker::Construct(const FArguments& InArgs)
 {
 	WorkspaceAssets = InArgs._WorkspaceAssets;
-
+	WorkspaceFactoryClass = InArgs._WorkspaceFactoryClass;
+	HintText = InArgs._HintText;
+	
 	FContentBrowserModule& ContentBrowserModule = FModuleManager::Get().LoadModuleChecked<FContentBrowserModule>(TEXT("ContentBrowser"));
 
 	FAssetPickerConfig AssetPickerConfig;
 	AssetPickerConfig.InitialAssetViewType = EAssetViewType::List;
-	AssetPickerConfig.OnAssetSelected = FOnAssetSelected::CreateLambda([this, OnAssetSelected = InArgs._OnAssetSelected](const FAssetData& InAssetData)
+	AssetPickerConfig.OnAssetSelected = FOnAssetSelected::CreateLambda([this, OnWorkspaceSelected = InArgs._OnExistingWorkspaceSelected](const FAssetData& InAssetData)
 	{
-		OnAssetSelected.ExecuteIfBound(InAssetData);
+		SelectedWorkspace = CastChecked<UWorkspace>(InAssetData.GetAsset());
+		OnWorkspaceSelected.ExecuteIfBound(SelectedWorkspace);
 
 		if(TSharedPtr<SWindow> Window = WeakWindow.Pin())
 		{
@@ -38,6 +44,19 @@ void SWorkspacePicker::Construct(const FArguments& InArgs)
 		return !WorkspaceAssets.IsEmpty() && !WorkspaceAssets.Contains(InAssetData);
 	});
 
+	const FText DisplayHintText = [this]()
+	{
+		switch (HintText)
+		{
+			case EHintText::SelectedAssetIsPartOfMultipleWorkspaces:
+				return LOCTEXT("SelectedAssetIsPartOfMultipleWorkspaces", "Selected asset is part of multiple workspaces.\nPlease select the workspace you want to open the asset with");
+			case EHintText::CreateOrUseExistingWorkspace:
+				return LOCTEXT("CreateOrUseExistingWorkspace", "Create or select an existing workspace to associate assets with.");
+			default:
+				return FText();
+		}
+	}();
+
 	ChildSlot
 	[
 		SNew(SBox)
@@ -51,7 +70,7 @@ void SWorkspacePicker::Construct(const FArguments& InArgs)
 			.HAlign(HAlign_Left)
 			[
 				SNew(STextBlock)
-				.Text(LOCTEXT("ChooseExistingWorkspaceToDetails", "Selected asset is part of multiple workspaces.\nPlease select the workspace you want to open the asset with"))
+				.Text(DisplayHintText)
 			]
 			+SVerticalBox::Slot()
 			.FillHeight(1.0f)
@@ -66,14 +85,28 @@ void SWorkspacePicker::Construct(const FArguments& InArgs)
 			[
 				SNew(SPrimaryButton)
 				.Text(LOCTEXT("CreateNewWorkspace", "Create New Workspace"))
-				.OnClicked_Lambda([this, OnNewAsset = InArgs._OnNewAsset]()
+				.OnClicked_Lambda([this, OnWorkspaceSelected = InArgs._OnNewWorkspaceCreated]()
 				{
-					OnNewAsset.ExecuteIfBound();
+					UWorkspaceFactory* Factory = NewObject<UWorkspaceFactory>(GetTransientPackage(), WorkspaceFactoryClass.Get());
+					UPackage* Package = CreatePackage(nullptr);
+					FName PackageName = *FPaths::GetBaseFilename(Package->GetName());
+					if (Factory->ConfigureProperties())
+					{
+						FAssetToolsModule& AssetToolsModule = FAssetToolsModule::GetModule();
+						UObject* NewAsset = AssetToolsModule.Get().CreateAssetWithDialog(UWorkspace::StaticClass(), Factory);
 
+						if (NewAsset)
+						{
+							SelectedWorkspace = CastChecked<UWorkspace>(NewAsset);
+							OnWorkspaceSelected.ExecuteIfBound(SelectedWorkspace);
+						}
+					}
+					
 					if(TSharedPtr<SWindow> Window = WeakWindow.Pin())
 					{
 						Window->RequestDestroyWindow();
 					}
+					
 					return FReply::Handled();
 				})
 			]
@@ -84,7 +117,7 @@ void SWorkspacePicker::Construct(const FArguments& InArgs)
 void SWorkspacePicker::ShowModal()
 {
 	TSharedRef<SWindow> Window = SNew(SWindow)
-	.Title(LOCTEXT("ChooseExistingWorkspaceToOpen", "Choose Existing Workspace to Open"))
+	.Title(LOCTEXT("ChooseWorkspaceToOpen", "Choose Workspace to Open"))
 	.SizingRule(ESizingRule::Autosized)
 	.SupportsMaximize(false)
 	.SupportsMinimize(false)
@@ -97,6 +130,11 @@ void SWorkspacePicker::ShowModal()
 	FSlateApplication::Get().AddModalWindow(Window, FGlobalTabmanager::Get()->GetRootWindow());
 }
 
+TObjectPtr<UObject> SWorkspacePicker::GetSelectedWorkspace() const
+{
+	return SelectedWorkspace;
+}
+	
 }
 
 #undef LOCTEXT_NAMESPACE
\ No newline at end of file
diff --git a/Engine/Plugins/Experimental/Workspace/Source/WorkspaceEditor/Private/SWorkspacePicker.h b/Engine/Plugins/Experimental/Workspace/Source/WorkspaceEditor/Private/SWorkspacePicker.h
--- a/Engine/Plugins/Experimental/Workspace/Source/WorkspaceEditor/Private/SWorkspacePicker.h
+++ b/Engine/Plugins/Experimental/Workspace/Source/WorkspaceEditor/Private/SWorkspacePicker.h
@@ -3,35 +3,51 @@
 #pragma once
 
 #include "ContentBrowserDelegates.h"
+#include "IWorkspacePicker.h"
 #include "Widgets/SCompoundWidget.h"
 
 class UWorkspace;
+class UWorkspaceFactory;
 
 namespace UE::Workspace
 {
 
-class SWorkspacePicker : public SCompoundWidget
+class SWorkspacePicker : public SCompoundWidget, public IWorkspacePicker
 {
-	using FOnNewAsset = TDelegate<void(TConstArrayView<FAssetData>)>; 
+	DECLARE_DELEGATE_OneParam(FOnWorkspaceSelected, TObjectPtr<UWorkspace>)
 	
 	SLATE_BEGIN_ARGS(SWorkspacePicker) {}
 
-	SLATE_EVENT(FOnAssetSelected, OnAssetSelected)
+	SLATE_EVENT(FOnWorkspaceSelected, OnExistingWorkspaceSelected)
 
-	SLATE_EVENT(FSimpleDelegate, OnNewAsset)
+	SLATE_EVENT(FOnWorkspaceSelected, OnNewWorkspaceCreated)
 
 	// The set of workspace assets the user can choose from
 	SLATE_ARGUMENT(TArray<FAssetData>, WorkspaceAssets)
 
+	SLATE_ARGUMENT(IWorkspacePicker::EHintText, HintText)
+		
+	SLATE_ARGUMENT(TSubclassOf<UWorkspaceFactory>, WorkspaceFactoryClass)
+
+	SLATE_ARGUMENT(bool, bShowSaveDialog)
+		
 	SLATE_END_ARGS()
 
 	void Construct(const FArguments& InArgs);
 
-	void ShowModal();
+	// IWorkspacePicker
+	virtual void ShowModal() override;
+	virtual TObjectPtr<UObject> GetSelectedWorkspace() const override;
 
 private:
 	TArray<FAssetData> WorkspaceAssets;
 
+	EHintText HintText = EHintText::SelectedAssetIsPartOfMultipleWorkspaces;
+
+	TSubclassOf<UWorkspaceFactory> WorkspaceFactoryClass;
+
+	TObjectPtr<UWorkspace> SelectedWorkspace;
+
 	TWeakPtr<SWindow> WeakWindow;
 };
 
diff --git a/Engine/Plugins/Experimental/Workspace/Source/WorkspaceEditor/Private/Workspace.h b/Engine/Plugins/Experimental/Workspace/Source/WorkspaceEditor/Private/Workspace.h
--- a/Engine/Plugins/Experimental/Workspace/Source/WorkspaceEditor/Private/Workspace.h
+++ b/Engine/Plugins/Experimental/Workspace/Source/WorkspaceEditor/Private/Workspace.h
@@ -28,12 +28,18 @@ namespace UE::Workspace
 	struct FWorkspaceOutliner;
 }
 
+namespace UE::UAF::Editor
+{
+	class FAssetWizard;
+}
+
 namespace UE::Workspace
 {
 	// A delegate for subscribing / reacting to workspace modifications.
 	DECLARE_MULTICAST_DELEGATE_OneParam(FOnWorkspaceModified, UWorkspace* /* InWorkspace */);
 }
 
+
 // Workspace entry used to export to asset registry
 USTRUCT()
 struct FWorkspaceAssetRegistryExportEntry
@@ -77,10 +83,12 @@ class UWorkspace : public UObject
 	friend class UAssetDefinition_Workspace;
 	friend class UWorkspaceState;
 
+	friend class UE::UAF::Editor::FAssetWizard;
+
 	// Adds an asset to the workspace
 	// @return true if the asset was added
 	UFUNCTION(BlueprintCallable, Category = "Workspace")
-	bool AddAsset(UObject* InAsset, bool bSetupUndoRedo = true, bool bPrintPythonCommand = true);
+	WORKSPACEEDITOR_API bool AddAsset(UObject* InAsset, bool bSetupUndoRedo = true, bool bPrintPythonCommand = true);
 	bool AddAsset(const FAssetData& InAsset, bool bSetupUndoRedo = true, bool bPrintPythonCommand = true);
 
 	// Adds assets to the workspace
diff --git a/Engine/Plugins/Experimental/Workspace/Source/WorkspaceEditor/Private/WorkspaceEditorModule.cpp b/Engine/Plugins/Experimental/Workspace/Source/WorkspaceEditor/Private/WorkspaceEditorModule.cpp
--- a/Engine/Plugins/Experimental/Workspace/Source/WorkspaceEditor/Private/WorkspaceEditorModule.cpp
+++ b/Engine/Plugins/Experimental/Workspace/Source/WorkspaceEditor/Private/WorkspaceEditorModule.cpp
@@ -239,53 +239,50 @@ bool FWorkspaceEditorModule::GetExportedAssetsForWorkspace(const FAssetData& InW
 
 	IWorkspaceEditor* WorkspaceEditor = nullptr;
 
-	auto HandleNewWorkspace = [InObject, &WorkspaceEditor, WorkSpaceFactoryClass]()
+	auto HandleNewWorkspace = [InObject, &WorkspaceEditor](TObjectPtr<UWorkspace> NewWorkspace)
 	{
-		UWorkspaceFactory* Factory = NewObject<UWorkspaceFactory>(GetTransientPackage(), WorkSpaceFactoryClass.Get());
-		UPackage* Package = CreatePackage(nullptr);
-		FName PackageName = *FPaths::GetBaseFilename(Package->GetName());
-		if (Factory->ConfigureProperties())
-		{
-			UWorkspace* NewWorkspace = CastChecked<UWorkspace>(Factory->FactoryCreateNew(UWorkspace::StaticClass(), Package, PackageName, RF_Public | RF_Standalone, nullptr, GWarn));
-			NewWorkspace->AddAsset(InObject, false);
-			NewWorkspace->MarkPackageDirty();
-
-			UAssetEditorSubsystem* AssetEditorSubsystem = GEditor->GetEditorSubsystem<UAssetEditorSubsystem>();
-			UWorkspaceAssetEditor* AssetEditor = NewObject<UWorkspaceAssetEditor>(AssetEditorSubsystem, NAME_None, RF_Transient);
-			AssetEditor->SetObjectToEdit(NewWorkspace);
-			AssetEditor->Initialize();
-			
-			WorkspaceEditor = static_cast<IWorkspaceEditor*>(GEditor->GetEditorSubsystem<UAssetEditorSubsystem>()->FindEditorForAsset(NewWorkspace, true));
-		}
+		NewWorkspace->AddAsset(InObject, false);
+		NewWorkspace->MarkPackageDirty();
+
+		UAssetEditorSubsystem* AssetEditorSubsystem = GEditor->GetEditorSubsystem<UAssetEditorSubsystem>();
+		UWorkspaceAssetEditor* AssetEditor = NewObject<UWorkspaceAssetEditor>(AssetEditorSubsystem, NAME_None, RF_Transient);
+		AssetEditor->SetObjectToEdit(NewWorkspace);
+		AssetEditor->Initialize();
+		
+		WorkspaceEditor = static_cast<IWorkspaceEditor*>(GEditor->GetEditorSubsystem<UAssetEditorSubsystem>()->FindEditorForAsset(NewWorkspace, true));
 	};
 
-	auto HandleExistingWorkspace = [&WorkspaceEditor](const FAssetData& InAssetData)
+	auto HandleExistingWorkspace = [InObject, &WorkspaceEditor](TObjectPtr<UWorkspace> ExistingWorkspace)
 	{
-		if(UWorkspace* ExistingWorkspace = Cast<UWorkspace>(InAssetData.GetAsset()))
-		{
-			GEditor->GetEditorSubsystem<UAssetEditorSubsystem>()->OpenEditorForAsset(ExistingWorkspace);
+		GEditor->GetEditorSubsystem<UAssetEditorSubsystem>()->OpenEditorForAsset(ExistingWorkspace);
 
-			WorkspaceEditor = static_cast<IWorkspaceEditor*>(GEditor->GetEditorSubsystem<UAssetEditorSubsystem>()->FindEditorForAsset(ExistingWorkspace, true));
-		}
+		WorkspaceEditor = static_cast<IWorkspaceEditor*>(GEditor->GetEditorSubsystem<UAssetEditorSubsystem>()->FindEditorForAsset(ExistingWorkspace, true));
 	};
 
 	if(InOpenMethod == EOpenWorkspaceMethod::AlwaysOpenNewWorkspace || RelevantWorkspaceAssets.Num() == 0)
 	{
 		// No relevant workspaces, so open a new one and add the asset
-		HandleNewWorkspace();
+		UWorkspaceFactory* Factory = NewObject<UWorkspaceFactory>(GetTransientPackage(), WorkSpaceFactoryClass.Get());
+		UPackage* Package = CreatePackage(nullptr);
+		FName PackageName = *FPaths::GetBaseFilename(Package->GetName());
+		if (Factory->ConfigureProperties())
+		{
+			TObjectPtr<UWorkspace> SelectedWorkspace = CastChecked<UWorkspace>(Factory->FactoryCreateNew(UWorkspace::StaticClass(), Package, PackageName, RF_Public | RF_Standalone, nullptr, GWarn));
+			HandleNewWorkspace(SelectedWorkspace);
+		}
 	}
 	else if(RelevantWorkspaceAssets.Num() == 1)
 	{
 		// One existing workspace, open it
-		HandleExistingWorkspace(RelevantWorkspaceAssets[0]);
+		HandleExistingWorkspace(CastChecked<UWorkspace>(RelevantWorkspaceAssets[0].GetAsset()));
 	}
 	else
 	{
 		// Multiple existing workspaces, present a window to let the user choose one to open with
 		TSharedRef<SWorkspacePicker> WorkspacePicker = SNew(SWorkspacePicker)
 			.WorkspaceAssets(RelevantWorkspaceAssets)
-			.OnAssetSelected_Lambda(HandleExistingWorkspace)
-			.OnNewAsset_Lambda(HandleNewWorkspace);
+			.OnExistingWorkspaceSelected_Lambda(HandleExistingWorkspace)
+			.OnNewWorkspaceCreated_Lambda(HandleNewWorkspace);
 
 		WorkspacePicker->ShowModal();
 	}
@@ -298,6 +295,15 @@ bool FWorkspaceEditorModule::GetExportedAssetsForWorkspace(const FAssetData& InW
 	return WorkspaceEditor;
 }
 
+ TSharedPtr<IWorkspacePicker> FWorkspaceEditorModule::CreateWorkspacePicker(const IWorkspacePicker::FConfig& Config) const
+ {
+	TSharedRef<SWorkspacePicker> WorkspacePicker = SNew(SWorkspacePicker)
+		.HintText(Config.HintText)
+		.WorkspaceFactoryClass(Config.WorkspaceFactoryClass);
+
+	return WorkspacePicker;
+ }
+
 const FObjectDocumentArgs* FWorkspaceEditorModule::FindObjectDocumentType(const FWorkspaceDocument& Document) const
 {
 	const FObjectDocumentArgs* FoundArgs = ObjectDocumentArgs.Find(Document.GetObject()->GetClass()->GetClassPathName());
diff --git a/Engine/Plugins/Experimental/Workspace/Source/WorkspaceEditor/Private/WorkspaceEditorModule.h b/Engine/Plugins/Experimental/Workspace/Source/WorkspaceEditor/Private/WorkspaceEditorModule.h
--- a/Engine/Plugins/Experimental/Workspace/Source/WorkspaceEditor/Private/WorkspaceEditorModule.h
+++ b/Engine/Plugins/Experimental/Workspace/Source/WorkspaceEditor/Private/WorkspaceEditorModule.h
@@ -39,6 +39,7 @@ private:
 	virtual FOnRegisterTabs& OnRegisterTabsForEditor() override { return RegisterTabsForEditor; }
 	virtual FOnExtendTabs& OnExtendTabs() override { return ExtendTabsForEditor; }
 	virtual FOnExtendToolMenuContext& OnExtendToolMenuContext() override { return ExtendToolMenuContext; }
+	virtual TSharedPtr<IWorkspacePicker> CreateWorkspacePicker(const IWorkspacePicker::FConfig& Config) const override;
 
 	// Find an existing registered object document type. Note this redirects based on FObjectDocumentArgs::OnRedirectWorkspaceContext
 	const FObjectDocumentArgs* FindObjectDocumentType(const FWorkspaceDocument& Document) const;
@@ -88,5 +89,5 @@ private:
 	friend SWorkspaceTabWrapper;
 	friend FWorkspaceEditorContext;
 };
-
+	
 }
diff --git a/Engine/Plugins/Experimental/Workspace/Source/WorkspaceEditor/Public/IWorkspaceEditorModule.h b/Engine/Plugins/Experimental/Workspace/Source/WorkspaceEditor/Public/IWorkspaceEditorModule.h
--- a/Engine/Plugins/Experimental/Workspace/Source/WorkspaceEditor/Public/IWorkspaceEditorModule.h
+++ b/Engine/Plugins/Experimental/Workspace/Source/WorkspaceEditor/Public/IWorkspaceEditorModule.h
@@ -10,6 +10,7 @@
 #include "Misc/Attribute.h"
 #include "Modules/ModuleInterface.h"
 #include "IWorkspaceOutlinerItemDetails.h"
+#include "IWorkspacePicker.h"
 
 #define UE_API WORKSPACEEDITOR_API
 
@@ -26,6 +27,7 @@ struct FTopLevelAssetPath;
 struct FWorkspaceDocumentState;
 struct FRegisterCustomClassLayoutParams;
 
+
 namespace UE::Workspace
 {
 	class IWorkspaceEditor;
@@ -267,6 +269,8 @@ public:
 	// Event to allow extending the FToolMenuContext
 	DECLARE_EVENT_TwoParams(IWorkspaceEditorModule, FOnExtendToolMenuContext, TSharedPtr<IWorkspaceEditor>, FToolMenuContext&);
 	virtual FOnExtendToolMenuContext& OnExtendToolMenuContext() = 0;
+
+	virtual TSharedPtr<IWorkspacePicker> CreateWorkspacePicker(const IWorkspacePicker::FConfig& Config) const = 0;
 };
 
 }
diff --git a/Engine/Plugins/Experimental/Workspace/Source/WorkspaceEditor/Public/IWorkspacePicker.h b/Engine/Plugins/Experimental/Workspace/Source/WorkspaceEditor/Public/IWorkspacePicker.h
new file mode 100644
--- /dev/null
+++ b/Engine/Plugins/Experimental/Workspace/Source/WorkspaceEditor/Public/IWorkspacePicker.h
@@ -0,0 +1,33 @@
+// Copyright Epic Games, Inc. All Rights Reserved.
+
+#pragma once
+
+#include "CoreMinimal.h"
+#include "WorkspaceFactory.h"
+
+class UWorkspace;
+
+namespace UE::Workspace
+{
+	class IWorkspacePicker
+	{
+	public:
+		virtual ~IWorkspacePicker() = default;
+
+		enum class EHintText
+		{
+			SelectedAssetIsPartOfMultipleWorkspaces,
+			CreateOrUseExistingWorkspace
+		};
+		
+		struct FConfig
+		{
+			EHintText HintText = EHintText::SelectedAssetIsPartOfMultipleWorkspaces;
+			TSubclassOf<UWorkspaceFactory> WorkspaceFactoryClass = UWorkspaceFactory::StaticClass();
+		};
+		
+		virtual void ShowModal() = 0;
+
+		virtual TObjectPtr<UObject> GetSelectedWorkspace() const = 0;
+	};
+}
\ No newline at end of file
diff --git a/Engine/Plugins/Experimental/Workspace/Source/WorkspaceEditor/Public/WorkspaceFactory.h b/Engine/Plugins/Experimental/Workspace/Source/WorkspaceEditor/Public/WorkspaceFactory.h
--- a/Engine/Plugins/Experimental/Workspace/Source/WorkspaceEditor/Public/WorkspaceFactory.h
+++ b/Engine/Plugins/Experimental/Workspace/Source/WorkspaceEditor/Public/WorkspaceFactory.h
@@ -14,6 +14,7 @@ namespace UE::Workspace
 {
 	struct FUtils;
 	class FWorkspaceEditorModule;
+	class SWorkspacePicker;
 }
 
 UCLASS(MinimalAPI, BlueprintType)
@@ -39,6 +40,7 @@ private:
 	TSubclassOf<UWorkspaceSchema> SchemaClass = nullptr;
 
 	friend class UE::Workspace::FWorkspaceEditorModule;
+	friend class UE::Workspace::SWorkspacePicker;
 };
 
 #undef UE_API
