Change FAnimNextAssetRegistryExports::ForEachType signatures to also pass in Identifier (FName) alongside the instanced struct, and use IsChildOf rather than exact class comparison

diff --git a/Engine/Plugins/Experimental/AnimNext/Source/AnimNextUncookedOnly/Public/AnimNextExports.h b/Engine/Plugins/Experimental/AnimNext/Source/AnimNextUncookedOnly/Public/AnimNextExports.h
--- a/Engine/Plugins/Experimental/AnimNext/Source/AnimNextUncookedOnly/Public/AnimNextExports.h
+++ b/Engine/Plugins/Experimental/AnimNext/Source/AnimNextUncookedOnly/Public/AnimNextExports.h
@@ -152,15 +152,15 @@ struct FAnimNextAssetRegistryExports
 	TArray<FAnimNextExport> Exports;
 
 	template<typename ExportType>
-	void ForEachExportOfType(TFunctionRef<bool(const ExportType&)> FunctionRef)
+	void ForEachExportOfType(TFunctionRef<bool(const FName&, const ExportType&)> FunctionRef)
 	{
 		for (const FAnimNextExport& Export : Exports)
 		{
 			if (Export.Data.IsValid())
 			{
-				if (Export.Data.GetScriptStruct() && Export.Data.GetScriptStruct() == ExportType::StaticStruct())
+				if (Export.Data.GetScriptStruct() && Export.Data.GetScriptStruct()->IsChildOf(ExportType::StaticStruct()))
 				{
-					if (!FunctionRef(Export.Data.Get<ExportType>()))
+					if (!FunctionRef(Export.Identifier, Export.Data.Get<ExportType>()))
 					{
 						return;
 					}
@@ -170,15 +170,15 @@ struct FAnimNextAssetRegistryExports
 	}
 
 	template<typename ExportType>
-	void ForEachExportOfType(TFunctionRef<bool(ExportType&)> FunctionRef)
+	void ForEachExportOfType(TFunctionRef<bool(const FName&, ExportType&)> FunctionRef)
 	{
 		for (const FAnimNextExport& Export : Exports)
 		{
 			if (Export.Data.IsValid())
 			{
-				if (Export.Data.GetScriptStruct() && Export.Data.GetScriptStruct() == ExportType::StaticStruct())
+				if (Export.Data.GetScriptStruct() && Export.Data.GetScriptStruct()->IsChildOf(ExportType::StaticStruct()))
 				{
-					if (!FunctionRef(Export.Data.Get<ExportType>()))
+					if (!FunctionRef(Export.Identifier, Export.Data.Get<ExportType>()))
 					{
 						return;
 					}
diff --git a/Engine/Plugins/Experimental/AnimNextAnimGraph/Source/AnimNextAnimGraphTestSuite/Private/AnimNextAnimGraphTraitGraphTest.cpp b/Engine/Plugins/Experimental/AnimNextAnimGraph/Source/AnimNextAnimGraphTestSuite/Private/AnimNextAnimGraphTraitGraphTest.cpp
--- a/Engine/Plugins/Experimental/AnimNextAnimGraph/Source/AnimNextAnimGraphTestSuite/Private/AnimNextAnimGraphTraitGraphTest.cpp
+++ b/Engine/Plugins/Experimental/AnimNextAnimGraph/Source/AnimNextAnimGraphTestSuite/Private/AnimNextAnimGraphTraitGraphTest.cpp
@@ -427,7 +427,7 @@ bool FAnimationAnimNextEditorTest_GraphManifest::RunTest(const FString& InParame
 		auto CountNumberOfManifestNodes = [&OutExports]()-> int32
 		{
 			int32 NumNodes = 0;			
-			OutExports.ForEachExportOfType<FAnimNextManifestNodeData>([&NumNodes](const FAnimNextManifestNodeData& Node) { ++NumNodes; return true; });
+			OutExports.ForEachExportOfType<FAnimNextManifestNodeData>([&NumNodes](const FName& Identifier, const FAnimNextManifestNodeData& Node) { ++NumNodes; return true; });
 			return NumNodes;
 		};
 		
