- Ensure that ModifiedDelegates are registered correctly when setting VariablesOutliner target export
- Refresh Assets shown in VariablesOutliner when an entry gets removed (removing SharedVariables has impact on reference chain)

diff --git a/Engine/Plugins/Experimental/AnimNext/Source/AnimNextEditor/Private/Variables/SVariablesView.cpp b/Engine/Plugins/Experimental/AnimNext/Source/AnimNextEditor/Private/Variables/SVariablesView.cpp
--- a/Engine/Plugins/Experimental/AnimNext/Source/AnimNextEditor/Private/Variables/SVariablesView.cpp
+++ b/Engine/Plugins/Experimental/AnimNext/Source/AnimNextEditor/Private/Variables/SVariablesView.cpp
@@ -92,12 +92,21 @@ void SVariablesOutliner::SetAssets(TConstArrayView<TSoftObjectPtr<UAnimNextRigVM
 	FullRefresh();
 }
 
-void SVariablesOutliner::SetExport(const FWorkspaceOutlinerItemExport& InExport)
+void SVariablesOutliner::UpdateAssets()
 {
-	Export = InExport;
-
 	TArray<FSoftObjectPath> AssetsToAsyncLoad;
 	TArray<TSoftObjectPtr<UAnimNextRigVMAsset>> ExportAssets;
+
+	auto HandleAssetPath = [this, &AssetsToAsyncLoad, &ExportAssets](const FSoftObjectPath& Path)
+	{
+		ExportAssets.Add(TSoftObjectPtr<UAnimNextRigVMAsset>(Path));
+
+		if (Path.ResolveObject() == nullptr)
+		{
+			AssetsToAsyncLoad.Add(Path);
+		}
+	};
+	
 	if (FWorkspaceOutlinerAssetReferenceItemData::IsAssetReference(Export))
 	{
 		// Reference item
@@ -106,12 +115,7 @@ void SVariablesOutliner::SetExport(const FWorkspaceOutlinerItemExport& InExport)
 
 		for (const FSoftObjectPath& Path : AssetPaths)
 		{
-			ExportAssets.Add(TSoftObjectPtr<UAnimNextRigVMAsset>(Path));
-			
-			if(Path.ResolveObject() == nullptr)
-			{
-				AssetsToAsyncLoad.Add(Path);
-			}
+			HandleAssetPath(Path);
 		}
 	}
 	else
@@ -119,13 +123,7 @@ void SVariablesOutliner::SetExport(const FWorkspaceOutlinerItemExport& InExport)
 		const FSoftObjectPath FirstAssetPath = Export.GetFirstAssetPath();
 		if(FirstAssetPath.IsAsset())
 		{
-			ExportAssets.Add(TSoftObjectPtr<UAnimNextRigVMAsset>(FirstAssetPath));
-
-			UObject* ResolvedObject = FirstAssetPath.ResolveObject();
-			if(ResolvedObject == nullptr)
-			{
-				AssetsToAsyncLoad.Add(FirstAssetPath);
-			}
+			HandleAssetPath(FirstAssetPath);
 		}
 		else if(Export.HasData() && Export.GetData().GetScriptStruct()->IsChildOf(FAnimNextAssetEntryOutlinerData::StaticStruct()))
 		{
@@ -134,7 +132,7 @@ void SVariablesOutliner::SetExport(const FWorkspaceOutlinerItemExport& InExport)
 			{
 				if(UAnimNextRigVMAsset* Asset = EntryData.GetEntry()->GetTypedOuter<UAnimNextRigVMAsset>())
 				{
-					ExportAssets.Add(TSoftObjectPtr<UAnimNextRigVMAsset>(Asset));
+					HandleAssetPath(Asset);
 				}
 			}
 		}
@@ -147,7 +145,7 @@ void SVariablesOutliner::SetExport(const FWorkspaceOutlinerItemExport& InExport)
 				{
 					if(UAnimNextRigVMAsset* Asset = EdGraph->GetTypedOuter<UAnimNextRigVMAsset>())
 					{
-						ExportAssets.Add(TSoftObjectPtr<UAnimNextRigVMAsset>(Asset));
+						HandleAssetPath(Asset);
 					}
 				}
 			}
@@ -180,6 +178,13 @@ void SVariablesOutliner::SetExport(const FWorkspaceOutlinerItemExport& InExport)
 	}
 }
 
+void SVariablesOutliner::SetExport(const FWorkspaceOutlinerItemExport& InExport)
+{
+	Export = InExport;
+
+	UpdateAssets();
+}
+
 void SVariablesOutliner::HandleAssetLoaded(const FSoftObjectPath& InSoftObjectPath, UAnimNextRigVMAsset* InAsset)
 {	
 	if(InAsset && Assets.Contains(TSoftObjectPtr<UAnimNextRigVMAsset>(InAsset)))
@@ -198,12 +203,14 @@ void SVariablesOutliner::OnEditorDataModified(UAnimNextRigVMAssetEditorData* InE
 	{
 	case EAnimNextEditorDataNotifType::UndoRedo:
 	case EAnimNextEditorDataNotifType::EntryAdded:
-	case EAnimNextEditorDataNotifType::EntryRemoved:		
 	case EAnimNextEditorDataNotifType::VariableCategoryChanged:
 	case EAnimNextEditorDataNotifType::CategoryAdded:
 	case EAnimNextEditorDataNotifType::CategoryChanged:
 		FullRefresh();
 		break;
+	case EAnimNextEditorDataNotifType::EntryRemoved:
+		UpdateAssets();
+		break;
 	default:
 		break;
 	}
diff --git a/Engine/Plugins/Experimental/AnimNext/Source/AnimNextEditor/Private/Variables/SVariablesView.h b/Engine/Plugins/Experimental/AnimNext/Source/AnimNextEditor/Private/Variables/SVariablesView.h
--- a/Engine/Plugins/Experimental/AnimNext/Source/AnimNextEditor/Private/Variables/SVariablesView.h
+++ b/Engine/Plugins/Experimental/AnimNext/Source/AnimNextEditor/Private/Variables/SVariablesView.h
@@ -44,6 +44,7 @@ public:
 protected:
 	
 	void SetAssets(TConstArrayView<TSoftObjectPtr<UAnimNextRigVMAsset>> InAssets);
+	void UpdateAssets();
 
 	void HandleAssetLoaded(const FSoftObjectPath& InSoftObjectPath, UAnimNextRigVMAsset* InAsset);
 	
