Return valid entry name for SharedVariablesEntry in case it is struct-based

diff --git a/Engine/Plugins/Experimental/AnimNext/Source/AnimNextUncookedOnly/Private/Entries/AnimNextSharedVariablesEntry.cpp b/Engine/Plugins/Experimental/AnimNext/Source/AnimNextUncookedOnly/Private/Entries/AnimNextSharedVariablesEntry.cpp
--- a/Engine/Plugins/Experimental/AnimNext/Source/AnimNextUncookedOnly/Private/Entries/AnimNextSharedVariablesEntry.cpp
+++ b/Engine/Plugins/Experimental/AnimNext/Source/AnimNextUncookedOnly/Private/Entries/AnimNextSharedVariablesEntry.cpp
@@ -37,7 +37,15 @@ void UAnimNextSharedVariablesEntry::Initialize(UAnimNextRigVMAssetEditorData* In
 
 FName UAnimNextSharedVariablesEntry::GetEntryName() const
 {
-	return Asset ? Asset->GetFName() : NAME_None;
+	switch (Type)
+	{
+	case EAnimNextSharedVariablesType::Asset:
+		return Asset ? Asset->GetFName() : NAME_None;
+	case EAnimNextSharedVariablesType::Struct:
+		return Struct ? Struct->GetFName() : NAME_None;
+	}
+	
+	return NAME_None;
 }
 
 FText UAnimNextSharedVariablesEntry::GetDisplayName() const
