Add ability to remove SharedVariables reference using lozenges (right click context menu and middle-mouse click to remove)

diff --git a/Engine/Plugins/Experimental/AnimNext/Source/AnimNextEditor/Private/AnimNextEditorStyle.h b/Engine/Plugins/Experimental/AnimNext/Source/AnimNextEditor/Private/AnimNextEditorStyle.h
--- a/Engine/Plugins/Experimental/AnimNext/Source/AnimNextEditor/Private/AnimNextEditorStyle.h
+++ b/Engine/Plugins/Experimental/AnimNext/Source/AnimNextEditor/Private/AnimNextEditorStyle.h
@@ -8,6 +8,7 @@
 #include "Styling/AppStyle.h"
 #include "Styling/CoreStyle.h"
 #include "Styling/SlateTypes.h"
+#include "Styling/StyleColors.h"
 #include "Interfaces/IPluginManager.h"
 
 #define IMAGE_PLUGIN_BRUSH_SVG( RelativePath, ... ) FSlateVectorImageBrush( FAnimNextEditorStyle::InContent(RelativePath, ".svg"), __VA_ARGS__)
@@ -38,6 +39,13 @@ public:
 			.SetFont(DEFAULT_FONT("Regular", 8))
 			.SetColorAndOpacity(FSlateColor::UseForeground()));
 
+		const FVector2D Icon14x14(14.0f, 14.0f);
+		const FButtonStyle CloseButtonStyle = FButtonStyle()
+		.SetNormal(CORE_IMAGE_BRUSH_SVG("Starship/Common/close-small", Icon14x14, FStyleColors::Foreground))
+		.SetPressed(CORE_IMAGE_BRUSH_SVG("Starship/Common/close-small", Icon14x14, FStyleColors::Foreground))
+		.SetHovered(CORE_IMAGE_BRUSH_SVG("Starship/Common/close-small", Icon14x14, FStyleColors::ForegroundHover));
+		Set("VariablesOutliner.SharedVariablesPill.CloseButton", CloseButtonStyle);
+
 		FSlateStyleRegistry::RegisterSlateStyle(*this);
 	}
 
diff --git a/Engine/Plugins/Experimental/AnimNext/Source/AnimNextEditor/Private/Variables/Outliner/VariablesOutlinerAssetItem.cpp b/Engine/Plugins/Experimental/AnimNext/Source/AnimNextEditor/Private/Variables/Outliner/VariablesOutlinerAssetItem.cpp
--- a/Engine/Plugins/Experimental/AnimNext/Source/AnimNextEditor/Private/Variables/Outliner/VariablesOutlinerAssetItem.cpp
+++ b/Engine/Plugins/Experimental/AnimNext/Source/AnimNextEditor/Private/Variables/Outliner/VariablesOutlinerAssetItem.cpp
@@ -18,6 +18,7 @@
 #include "VariablesOutlinerDragDrop.h"
 #include "VariablesOutlinerMode.h"
 #include "WorkspaceDragDropOperation.h"
+#include "Entries/AnimNextSharedVariablesEntry.h"
 #include "Styling/SlateIconFinder.h"
 #include "Styling/StarshipCoreStyle.h"
 #include "UObject/Package.h"
@@ -30,6 +31,7 @@
 #include "Widgets/Text/SInlineEditableTextBlock.h"
 #include "Widgets/Text/STextBlock.h"
 #include "Widgets/Views/STableRow.h"
+#include "Widgets/Layout/SSpacer.h"
 
 #define LOCTEXT_NAMESPACE "VariablesOutlinerTreeItem"
 
@@ -161,6 +163,7 @@ class SVariablesOutlinerAssetLabel : FSceneOutlinerCommonLabelData, public SComp
 				const FString SharedVariablesAssetName = SharedVariablesPath.GetAssetName();
 				const FLinearColor BackgroundColor = FLinearColor::MakeRandomSeededColor(GetTypeHash(SharedVariablesAssetName)).Desaturate(0.4f);
 
+				const FButtonStyle& CloseButtonStyle = FAnimNextEditorStyle::Get().GetWidgetStyle<FButtonStyle>("VariablesOutliner.SharedVariablesPill.CloseButton");
 				TSharedPtr<SBorder> SharedVariablesBorderWidget;
 				ImplementedSharedVariablesHBox->AddSlot()
 				.VAlign(VAlign_Center)
@@ -172,14 +175,37 @@ class SVariablesOutlinerAssetLabel : FSceneOutlinerCommonLabelData, public SComp
 					.BorderBackgroundColor(BackgroundColor)
 					.ToolTipText( FText::FromString(SharedVariablesAssetName))
 					.VAlign(VAlign_Center)
-					.Padding(FMargin(5.0f, 1.0f, 5.0f, 1.0f))
+					.Padding(FMargin(6.0f, 1.0f, 4.0f, 1.0f))
 					[
-						SNew(STextBlock)
-						.TextStyle(FAnimNextEditorStyle::Get(), "VariablesOutliner.SharedVariablesPill.Text")
-						.Text(FText::FromString(SharedVariablesAssetName))
-						.HighlightText(SceneOutliner.GetFilterHighlightText())						
-						.OverflowPolicy(ETextOverflowPolicy::MiddleEllipsis)
+						SNew(SHorizontalBox)
+						+SHorizontalBox::Slot()
+						.AutoWidth()
+						.HAlign(HAlign_Center)
+						.VAlign(VAlign_Center)
+						[
+							SNew(STextBlock)
+							.TextStyle(FAnimNextEditorStyle::Get(), "VariablesOutliner.SharedVariablesPill.Text")
+							.Text(FText::FromString(SharedVariablesAssetName))
+							.HighlightText(SceneOutliner.GetFilterHighlightText())
+							.OverflowPolicy(ETextOverflowPolicy::MiddleEllipsis)
+						]
+						+SHorizontalBox::Slot()
+						.HAlign(HAlign_Right)
+						.VAlign(VAlign_Center)
+						.AutoWidth()
+						[
+							SNew(SButton)
+							.ButtonStyle(&CloseButtonStyle)
+							.OnClicked( Item.ToSharedRef(), &FVariablesOutlinerAssetItem::OnRemoveSharedVariable, SharedVariablesPath )
+							.ContentPadding(FMargin(0.0))
+							.ToolTipText(LOCTEXT("RemoveReference_ToolTip", "Removes the SharedVariables reference from this Asset."))
+							[
+								SNew(SSpacer)
+								.Size(CloseButtonStyle.Normal.ImageSize)
+							]
+						]
 					]
+					.OnMouseButtonUp(Item.ToSharedRef(), &FVariablesOutlinerAssetItem::OnSharedVariableWidgetMouseUp, SharedVariablesPath)
 				];
 
 				SharedVariablesBorderWidget->SetOnMouseMove(FPointerEventHandler::CreateLambda([SharedVariablesPath, WeakOutliner = WeakSceneOutliner](const FGeometry&, const FPointerEvent& Event) -> FReply
@@ -571,9 +597,9 @@ FString FVariablesOutlinerAssetItem::GetDisplayString() const
 
 TSharedRef<SWidget> FVariablesOutlinerAssetItem::GenerateLabelWidget(ISceneOutliner& Outliner, const STableRow<FSceneOutlinerTreeItemPtr>& InRow)
 {
-	TSharedRef<SVariablesOutlinerAssetLabel> LabelWidget = SNew(SVariablesOutlinerAssetLabel, *this, Outliner, InRow);
-	RenameRequestEvent.BindSP(LabelWidget->TextBlock.Get(), &SInlineEditableTextBlock::EnterEditingMode);
-	return LabelWidget;
+	SAssignNew(AssetLabel, SVariablesOutlinerAssetLabel, *this, Outliner, InRow);
+	RenameRequestEvent.BindSP(AssetLabel->TextBlock.Get(), &SInlineEditableTextBlock::EnterEditingMode);
+	return AssetLabel->AsShared();
 }
 
 FString FVariablesOutlinerAssetItem::GetPackageName() const
@@ -637,6 +663,86 @@ bool FVariablesOutlinerAssetItem::ValidateName(const FText& InNewName, FText& Ou
 	return true;
 }
 
+FReply FVariablesOutlinerAssetItem::OnSharedVariableWidgetMouseUp(const FGeometry& Geometry, const FPointerEvent& PointerEvent, const FSoftObjectPath ClickedSharedVariablesPath) const
+{
+	if (PointerEvent.GetEffectingButton() == EKeys::RightMouseButton)
+	{
+		const FWidgetPath WidgetPath = (PointerEvent.GetEventPath() != nullptr) ? *PointerEvent.GetEventPath() : FWidgetPath();
+
+		const bool bCloseAfterSelection = true;
+		FMenuBuilder MenuBuilder(bCloseAfterSelection, nullptr);
+		MenuBuilder.BeginSection("", LOCTEXT("SharedVariablesContextMenu", "Shared Variables") );
+		{
+			MenuBuilder.AddMenuEntry(
+					LOCTEXT("RemoveReference", "Remove"),
+					LOCTEXT("RemoveReference_ToolTip", "Removes the SharedVariables reference from this Asset."),
+					FSlateIcon(),
+					FUIAction(FExecuteAction::CreateLambda([this, ClickedSharedVariablesPath]() { OnRemoveSharedVariable(ClickedSharedVariablesPath); }))
+			);
+		}
+		MenuBuilder.EndSection();
+
+		FSlateApplication::Get().PushMenu(AssetLabel.ToSharedRef(), WidgetPath, MenuBuilder.MakeWidget(), FSlateApplication::Get().GetCursorPos(), FPopupTransitionEffect(FPopupTransitionEffect::ContextMenu));
+		
+		return FReply::Handled();
+	}
+
+	return FReply::Unhandled();
+}
+
+FReply FVariablesOutlinerAssetItem::OnRemoveSharedVariable(const FSoftObjectPath ClickedSharedVariablesPath) const
+{
+	if (UAnimNextRigVMAsset* Asset = SoftAsset.Get())
+	{
+		const FText AssetNameText = FText::FromName(Asset->GetFName());
+
+		UAnimNextRigVMAssetEditorData* EditorData = UncookedOnly::FUtils::GetEditorData<UAnimNextRigVMAssetEditorData>(Asset);
+		if (UAnimNextSharedVariablesEntry* SharedVariableEntryToRemove = Cast<UAnimNextSharedVariablesEntry>(EditorData->FindEntry(ClickedSharedVariablesPath.GetAssetFName())))
+		{
+			const FText SharedVariablesAssetNameText = FText::FromName(ClickedSharedVariablesPath.GetAssetFName());
+			const FText FormattedMessage = FText::Format(LOCTEXT("RemovingSharedAssetEntryFormat", "Removing {0} from {1}"), SharedVariablesAssetNameText, AssetNameText);
+			FScopedTransaction Transaction(FormattedMessage);
+			
+			if (const UAnimNextSharedVariables* SharedVariablesToRemove = SharedVariableEntryToRemove->GetAsset())
+			{
+				// Replace all referenced variables from the to-be-removed SharedVariables from the referencing Asset
+				FAnimNextAssetRegistryExports VariableExports;
+				UncookedOnly::FUtils::GetExportedVariablesForAsset(SharedVariablesToRemove, VariableExports);
+
+				VariableExports.ForEachExportOfType<FAnimNextVariableDeclarationData>([EditorData, ClickedSharedVariablesPath](const FName& Identifier, const FAnimNextVariableDeclarationData& VariableDeclaration) -> bool
+				{
+					FAnimNextSoftVariableReference VariableReference(Identifier, ClickedSharedVariablesPath);
+					FAnimNextSoftVariableReference EmptyReference;
+					UncookedOnly::FUtils::ReplaceVariableReferences(EditorData, VariableReference, EmptyReference);
+				
+					return true;
+				});
+			}
+			else if (const UScriptStruct* Struct = SharedVariableEntryToRemove->GetStruct())
+			{
+				// Replace all referenced variables from the to-be-removed Native SharedVariables source from the referencing Asset
+				for (TFieldIterator<FProperty> It(Struct); It; ++It)
+				{
+					const FProperty* Property = *It;
+					if (Property->HasAnyPropertyFlags(CPF_NativeAccessSpecifierPublic))
+					{
+
+						FAnimNextSoftVariableReference VariableReference(Property->GetFName(), ClickedSharedVariablesPath);
+						FAnimNextSoftVariableReference EmptyReference;
+						UncookedOnly::FUtils::ReplaceVariableReferences(EditorData, VariableReference, EmptyReference);
+					}
+				}
+			}
+
+			// Remove the entry itself
+			EditorData->RemoveEntry(SharedVariableEntryToRemove);
+
+			return FReply::Handled();
+		}
+	}
+
+	return FReply::Unhandled();
+}
 }
 
 #undef LOCTEXT_NAMESPACE // "VariablesOutlinerTreeItem"
\ No newline at end of file
diff --git a/Engine/Plugins/Experimental/AnimNext/Source/AnimNextEditor/Private/Variables/Outliner/VariablesOutlinerAssetItem.h b/Engine/Plugins/Experimental/AnimNext/Source/AnimNextEditor/Private/Variables/Outliner/VariablesOutlinerAssetItem.h
--- a/Engine/Plugins/Experimental/AnimNext/Source/AnimNextEditor/Private/Variables/Outliner/VariablesOutlinerAssetItem.h
+++ b/Engine/Plugins/Experimental/AnimNext/Source/AnimNextEditor/Private/Variables/Outliner/VariablesOutlinerAssetItem.h
@@ -41,12 +41,17 @@ struct FVariablesOutlinerAssetItem : ISceneOutlinerTreeItem
 	// Validates the new item name
 	bool ValidateName(const FText& InNewName, FText& OutErrorMessage) const;
 
+	FReply OnSharedVariableWidgetMouseUp(const FGeometry&, const FPointerEvent& PointerEvent, const FSoftObjectPath ClickedSharedVariablesPath) const;
+	FReply OnRemoveSharedVariable(const FSoftObjectPath ClickedSharedVariablesPath) const;
+
 	// Soft ptr to the underlying asset, which may not be loaded yet
 	TSoftObjectPtr<UAnimNextRigVMAsset> SoftAsset;
 
 	int32 SortValue;
 
 	TArray<const FSoftObjectPath> SharedVariableSourcePaths;
+
+	TSharedPtr<class SVariablesOutlinerAssetLabel> AssetLabel;
 };
 
 }
diff --git a/Engine/Plugins/Experimental/AnimNext/Source/AnimNextUncookedOnly/Internal/AnimNextRigVMAssetEditorData.h b/Engine/Plugins/Experimental/AnimNext/Source/AnimNextUncookedOnly/Internal/AnimNextRigVMAssetEditorData.h
--- a/Engine/Plugins/Experimental/AnimNext/Source/AnimNextUncookedOnly/Internal/AnimNextRigVMAssetEditorData.h
+++ b/Engine/Plugins/Experimental/AnimNext/Source/AnimNextUncookedOnly/Internal/AnimNextRigVMAssetEditorData.h
@@ -62,6 +62,7 @@ namespace UE::UAF::Editor
 	class FAssetCompilationHandler;
 	class FVariableBindingPropertyCustomization;
 	struct FVariablesOutlinerCategoryItem;
+	struct FVariablesOutlinerAssetItem;
 }
 
 namespace UE::UAF::Tests
@@ -228,6 +229,7 @@ protected:
 	friend class UE::UAF::Editor::FAnimNextFunctionItemDetails;
 	friend struct UE::UAF::Editor::FVariablesOutlinerEntryItem;
 	friend struct UE::UAF::Editor::FVariablesOutlinerCategoryItem;
+	friend struct UE::UAF::Editor::FVariablesOutlinerAssetItem;
 	friend class UE::UAF::Editor::FVariablesOutlinerMode;
 	friend class UE::UAF::Editor::FVariablesOutlinerHierarchy;
 	friend class UE::UAF::Editor::SVariablesOutlinerValue;
diff --git a/Engine/Plugins/Experimental/AnimNext/Source/AnimNextUncookedOnly/Internal/Entries/AnimNextSharedVariablesEntry.h b/Engine/Plugins/Experimental/AnimNext/Source/AnimNextUncookedOnly/Internal/Entries/AnimNextSharedVariablesEntry.h
--- a/Engine/Plugins/Experimental/AnimNext/Source/AnimNextUncookedOnly/Internal/Entries/AnimNextSharedVariablesEntry.h
+++ b/Engine/Plugins/Experimental/AnimNext/Source/AnimNextUncookedOnly/Internal/Entries/AnimNextSharedVariablesEntry.h
@@ -22,6 +22,7 @@ namespace UE::UAF::Editor
 	class SVariableOverride;
 	class FVariablesOutlinerMode;
 	class SVariablesOutliner;
+	struct FVariablesOutlinerAssetItem;
 }
 
 namespace UE::UAF::UncookedOnly
@@ -53,6 +54,7 @@ class UAnimNextSharedVariablesEntry : public UAnimNextRigVMAssetEntry
 	friend class UE::UAF::Editor::SAddVariablesDialog;
 	friend class UE::UAF::Editor::SVariableOverride;
 	friend class UE::UAF::Editor::SVariablesOutliner;
+	friend struct UE::UAF::Editor::FVariablesOutlinerAssetItem;
 	friend UE::UAF::Editor::FVariablesOutlinerMode;
 	friend class UAnimNextAssetFindReplaceVariables;
 
